Parallel executor for GeometryFlow based on TaskGraph. Unit tests for GeometryFlow.

#rnx

diff --git a/Engine/Plugins/Experimental/GeometryFlow/Source/GeometryFlowCore/Private/GeometryFlowExecutor.cpp b/Engine/Plugins/Experimental/GeometryFlow/Source/GeometryFlowCore/Private/GeometryFlowExecutor.cpp
new file mode 100644
--- /dev/null
+++ b/Engine/Plugins/Experimental/GeometryFlow/Source/GeometryFlowCore/Private/GeometryFlowExecutor.cpp
@@ -0,0 +1,217 @@
+// Copyright Epic Games, Inc. All Rights Reserved.
+
+#include "GeometryFlowExecutor.h"
+#include "GeometryFlowGraphUtil.h"
+
+using namespace UE::GeometryFlow;
+
+
+namespace
+{
+	template<typename T>
+	bool PopElementFromSet(TSet<T>& Set, T& Element)
+	{
+		typename TSet<T>::TIterator Iter(Set);
+		if (!Iter)
+		{
+			return false;
+		}
+		else
+		{
+			Element = *Iter;
+			Iter.RemoveCurrent();
+			return true;
+		}
+	}
+}
+
+FGeometryFlowExecutor::FGeometryFlowExecutor(FGraph& InGraph)
+	: GeometryFlowGraph(InGraph)
+{
+	TopologicalSort();
+}
+
+
+void FGeometryFlowExecutor::ComputeOutputs(const TArray<NodeOutputSpec>& DesiredOutputs,
+										   TArray<TSafeSharedPtr<IData>>& OutputDatas)
+{
+	// TODO: Improve this!
+
+	CreateAndDispatchTaskGraph();
+
+	for (const NodeOutputSpec& OutputSpec : DesiredOutputs)
+	{
+		for (const TPair<FGraph::FHandle, FGraphEventRef>& NodeRef : GeometryFlowNodeToGraphTask)
+		{
+			FGraph::FHandle NodeHandle{ NodeRef.Key };
+			if ( OutputSpec.Key == NodeHandle)
+			{ 
+				// TODO: Don't wait for each one, see if we can set a callback for when a task is done
+				NodeRef.Value->Wait();
+
+				// TODO: Check desired data type matches node output type
+
+				TSafeSharedPtr<FNode> GeoFlowNode = GeometryFlowGraph.FindNode(NodeHandle);
+				TSafeSharedPtr<IData> OutputData = GeoFlowNode->GetOutput(OutputSpec.Value);
+				check(OutputData != nullptr);
+
+				OutputDatas.Add(OutputData);
+			}
+		}
+	}
+
+}
+
+
+
+
+void FGeometryFlowExecutor::TopologicalSort()
+{
+	// Kahn's algorithm from Wikipedia
+
+	TArray<FGraph::FConnection> Edges = GeometryFlowGraph.Connections;
+
+	// L := Empty list that will contain the sorted elements
+	TopologicallySortedNodes.Empty();
+
+	// S := Set of all nodes with no incoming edge
+	TSet<FGraph::FHandle> CurrentNodes = GeometryFlowGraph.GetSourceNodes();
+
+	while (CurrentNodes.Num() != 0)
+	{
+		// remove a node n from S
+		FGraph::FHandle N;
+		bool bPopOK = PopElementFromSet(CurrentNodes, N);
+		check(bPopOK);
+
+		// add n to L
+		TopologicallySortedNodes.Add(N);
+
+		// for each node m with an edge e from n to m do
+		int FoundEdgeIndex = FindAnyConnectionFromNode(N, Edges);
+		while (FoundEdgeIndex >= 0)
+		{
+			FGraph::FHandle M = Edges[FoundEdgeIndex].ToNode;
+
+			// remove edge e from the graph
+			Edges.RemoveAt(FoundEdgeIndex);
+
+			// if m has no other incoming edges then
+			if (FindAnyConnectionToNode(M, Edges) < 0)
+			{
+				// insert m into S
+				CurrentNodes.Add(M);
+			}
+
+			FoundEdgeIndex = FindAnyConnectionFromNode(N, Edges);
+		}
+	}
+
+	check(Edges.Num() == 0);		// If any edges remain there's a cycle in the graph
+}
+
+
+void FGeometryFlowExecutor::CreateAndDispatchTaskGraph()
+{
+	DebugNodeExecutionLog.Reset();
+	DebugNodeExecutionTime.Reset();
+
+	// Track nodes as they are added to the TaskGraph
+	GeometryFlowNodeToGraphTask.Reset();
+
+	// We'll add nodes to the task graph in topologically sorted order. This means any prereqs for a node should 
+	// already be in this container when the node is processed.
+
+	for (FGraph::FHandle NodeHandle : TopologicallySortedNodes)
+	{
+		// Find prereqs for the current node
+		TArray<int32> IncomingConnections = FindAllConnectionsToNode(NodeHandle, GeometryFlowGraph.Connections);
+		FGraphEventArray Prereqs;
+		for (int32 ConnectionIndex : IncomingConnections)
+		{
+			check(GeometryFlowGraph.Connections[ConnectionIndex].ToNode == NodeHandle);
+			FGraph::FHandle M = GeometryFlowGraph.Connections[ConnectionIndex].FromNode;
+			check(GeometryFlowNodeToGraphTask.Find(M) != nullptr);		// should have been already added to TaskGraph
+			Prereqs.Add(GeometryFlowNodeToGraphTask[M]);
+		}
+
+		// Now construct the node and add it to the list
+		
+		auto TaskLambda = [this, NodeHandle]()
+		{
+			FGraph::FNodeInfo* NodeInfo = GeometryFlowGraph.AllNodes.Find(NodeHandle);
+			check(NodeInfo != nullptr);
+
+			TSafeSharedPtr<FNode> Node = NodeInfo->Node;
+			check(Node != nullptr);
+
+			// this is the map of (InputName, Data) we will build up by pulling from the Connections
+			FNamedDataMap DatasIn;
+
+			// Collect data from those Inputs.
+			Node->EnumerateInputs([this, NodeHandle, &DatasIn, &Node] (const FString& InputName, const TUniquePtr<INodeInput>& Input)
+			{
+				// find the connection for this input
+				FGraph::FConnection Connection;
+				EGeometryFlowResult FoundResult = GeometryFlowGraph.FindConnectionForInput(NodeHandle, InputName, Connection);
+				check(FoundResult == EGeometryFlowResult::Ok);
+
+				// If there is only one Connection from this upstream Output (ie to our Input), and the Node/Input 
+				// can steal and transform that data, then we will do it
+				int32 OutputUsageCount = GeometryFlowGraph.CountOutputConnections(Connection.FromNode, Connection.FromOutput);
+				bool bStealDataForInput = false;
+				FDataFlags DataFlags;
+
+				ENodeCachingStrategy FromCachingStrategy = GeometryFlowGraph.GetCachingStrategyForNode(Connection.FromNode);
+				if (OutputUsageCount == 1
+					&& Input->CanTransformInput()
+					&& FromCachingStrategy != ENodeCachingStrategy::AlwaysCache)
+				{
+					bStealDataForInput = true;
+					DataFlags.bIsMutableData = true;
+				}
+
+				// TODO: Error if upstream node hasn't executed yet.
+
+				FGraph::FHandle UpstreamNodeHandle = Connection.FromNode;
+				TSafeSharedPtr<FNode> UpstreamNode = GeometryFlowGraph.FindNode(UpstreamNodeHandle);
+				FString UpstreamOutputName = Connection.FromOutput;
+				TSafeSharedPtr<IData> OutputData = UpstreamNode->GetOutput(UpstreamOutputName);
+
+				TSafeSharedPtr<IData> Result = (bStealDataForInput) ?
+					UpstreamNode->StealOutput(UpstreamOutputName) : UpstreamNode->GetOutput(UpstreamOutputName);
+
+				// add to named data map
+				DatasIn.Add(InputName, OutputData, DataFlags);
+			});
+
+			// evalute node
+			FNamedDataMap DatasOut;
+			for (const FNode::FNodeOutputInfo& NodeOutput : Node->NodeOutputs)
+			{
+				DatasOut.Add(NodeOutput.Name);
+			}
+
+			double T0 = FPlatformTime::Seconds();
+			
+			TUniquePtr<FEvaluationInfo> EvalInfo = MakeUnique<FEvaluationInfo>();
+			Node->Evaluate(DatasIn, DatasOut, EvalInfo);
+
+			double T1 = FPlatformTime::Seconds();
+
+			// Debug logging
+			DebugNodeExecutionLogLock.Lock();
+			DebugNodeExecutionLog.Add(NodeHandle.Identifier);
+			DebugNodeExecutionTime.Add(Node.Get(), T1 - T0);
+			DebugNodeExecutionLogLock.Unlock();
+		};
+		
+		// TODO: Create but don't dispatch? Possible?
+		FGraphEventRef TaskGraphN = FFunctionGraphTask::CreateAndDispatchWhenReady(MoveTemp(TaskLambda), TStatId{}, &Prereqs);
+
+		TPair<FGraph::FHandle, FGraphEventRef> NodeInfo{ NodeHandle, TaskGraphN };
+		GeometryFlowNodeToGraphTask.Add(NodeInfo);
+
+	}
+}
+
diff --git a/Engine/Plugins/Experimental/GeometryFlow/Source/GeometryFlowCore/Private/GeometryFlowGraph.cpp b/Engine/Plugins/Experimental/GeometryFlow/Source/GeometryFlowCore/Private/GeometryFlowGraph.cpp
--- a/Engine/Plugins/Experimental/GeometryFlow/Source/GeometryFlowCore/Private/GeometryFlowGraph.cpp
+++ b/Engine/Plugins/Experimental/GeometryFlow/Source/GeometryFlowCore/Private/GeometryFlowGraph.cpp
@@ -7,7 +7,7 @@ using namespace UE::GeometryFlow;
 
 TSafeSharedPtr<FNode> FGraph::FindNode(FHandle NodeHandle) const
 {
-	const FNodeInfo* Found = AllNodes.Find(NodeHandle.Identifier);
+	const FNodeInfo* Found = AllNodes.Find(NodeHandle);
 	if (!ensure(Found != nullptr))
 	{
 		return nullptr;
@@ -17,7 +17,7 @@ TSafeSharedPtr<FNode> FGraph::FindNode(FHandle NodeHandle) const
 
 EGeometryFlowResult FGraph::GetInputTypeForNode(FHandle NodeHandle, FString InputName, int32& Type) const
 {
-	const FNodeInfo* Found = AllNodes.Find(NodeHandle.Identifier);
+	const FNodeInfo* Found = AllNodes.Find(NodeHandle);
 	if (!ensure(Found != nullptr))
 	{
 		return EGeometryFlowResult::NodeDoesNotExist;
@@ -28,7 +28,7 @@ EGeometryFlowResult FGraph::GetInputTypeForNode(FHandle NodeHandle, FString Inpu
 
 EGeometryFlowResult FGraph::GetOutputTypeForNode(FHandle NodeHandle, FString OutputName, int32& Type) const
 {
-	const FNodeInfo* Found = AllNodes.Find(NodeHandle.Identifier);
+	const FNodeInfo* Found = AllNodes.Find(NodeHandle);
 	if (!ensure(Found != nullptr))
 	{
 		return EGeometryFlowResult::NodeDoesNotExist;
@@ -40,7 +40,7 @@ EGeometryFlowResult FGraph::GetOutputTypeForNode(FHandle NodeHandle, FString Out
 
 ENodeCachingStrategy FGraph::GetCachingStrategyForNode(FHandle NodeHandle) const
 {
-	const FNodeInfo* Found = AllNodes.Find(NodeHandle.Identifier);
+	const FNodeInfo* Found = AllNodes.Find(NodeHandle);
 	if (!ensure(Found != nullptr))
 	{
 		return ENodeCachingStrategy::AlwaysCache;
@@ -234,7 +234,7 @@ void FGraph::ConfigureCachingStrategy(ENodeCachingStrategy NewStrategy)
 
 EGeometryFlowResult FGraph::SetNodeCachingStrategy(FHandle NodeHandle, ENodeCachingStrategy Strategy)
 {
-	FNodeInfo* Found = AllNodes.Find(NodeHandle.Identifier);
+	FNodeInfo* Found = AllNodes.Find(NodeHandle);
 	if (!ensure(Found != nullptr))
 	{
 		return EGeometryFlowResult::NodeDoesNotExist;
diff --git a/Engine/Plugins/Experimental/GeometryFlow/Source/GeometryFlowCore/Private/GeometryFlowNode.cpp b/Engine/Plugins/Experimental/GeometryFlow/Source/GeometryFlowCore/Private/GeometryFlowNode.cpp
--- a/Engine/Plugins/Experimental/GeometryFlow/Source/GeometryFlowCore/Private/GeometryFlowNode.cpp
+++ b/Engine/Plugins/Experimental/GeometryFlow/Source/GeometryFlowCore/Private/GeometryFlowNode.cpp
@@ -129,6 +129,15 @@ void FNode::CollectRequirements(const TArray<FString>& Outputs, TArray<FEvalRequ
 }
 
 
+void FNode::CollectAllRequirements(TArray<FEvalRequirement>& RequiredInputsOut)
+{
+	RequiredInputsOut.Reset();
+	for (const FNodeInputInfo& InputInfo : NodeInputs)
+	{
+		FEvalRequirement Requirement(InputInfo.Name, InputInfo.Input->GetInputFlags());
+		RequiredInputsOut.Add(Requirement);
+	}
+}
 
 
 void FNode::ConfigureInputFlags(const FString& InputName, FNodeInputFlags Flags)
diff --git a/Engine/Plugins/Experimental/GeometryFlow/Source/GeometryFlowCore/Public/GeometryFlowCoreNodes.h b/Engine/Plugins/Experimental/GeometryFlow/Source/GeometryFlowCore/Public/GeometryFlowCoreNodes.h
--- a/Engine/Plugins/Experimental/GeometryFlow/Source/GeometryFlowCore/Public/GeometryFlowCoreNodes.h
+++ b/Engine/Plugins/Experimental/GeometryFlow/Source/GeometryFlowCore/Public/GeometryFlowCoreNodes.h
@@ -48,6 +48,11 @@ public:
 		return;
 	}
 
+	virtual void CollectAllRequirements(TArray<FEvalRequirement>& RequiredInputsOut) override
+	{
+		return;
+	}
+
 	virtual void Evaluate(
 		const FNamedDataMap& DatasIn,
 		FNamedDataMap& DatasOut,
diff --git a/Engine/Plugins/Experimental/GeometryFlow/Source/GeometryFlowCore/Public/GeometryFlowExecutor.h b/Engine/Plugins/Experimental/GeometryFlow/Source/GeometryFlowCore/Public/GeometryFlowExecutor.h
new file mode 100644
--- /dev/null
+++ b/Engine/Plugins/Experimental/GeometryFlow/Source/GeometryFlowCore/Public/GeometryFlowExecutor.h
@@ -0,0 +1,43 @@
+
+#pragma once
+
+#include "GeometryFlowGraph.h"
+#include "Async/TaskGraphInterfaces.h"
+
+namespace UE
+{
+namespace GeometryFlow
+{
+
+///
+/// Parallel execution using TaskGraph
+///
+class GEOMETRYFLOWCORE_API FGeometryFlowExecutor
+{
+
+public:
+
+	FGeometryFlowExecutor(FGraph& InGraph);
+
+	using NodeOutputSpec = TPair<FGraph::FHandle, FString>;		// (Node handle, output name)
+	void ComputeOutputs(const TArray<NodeOutputSpec>& DesiredOutputs, TArray<TSafeSharedPtr<IData>>& OutputDatas);
+
+protected:
+
+	FGraph GeometryFlowGraph;
+
+	TArray<FGraph::FHandle> TopologicallySortedNodes;
+	void TopologicalSort();
+
+	TMap<FGraph::FHandle, FGraphEventRef> GeometryFlowNodeToGraphTask;
+	void CreateAndDispatchTaskGraph();
+
+	// Debug helpers
+	TArray<int32> DebugNodeExecutionLog;		        // when a node executes it adds its ID to this array
+	TMap<FNode*, double> DebugNodeExecutionTime;		// log execution time
+	FCriticalSection DebugNodeExecutionLogLock;
+
+};
+
+}
+}
diff --git a/Engine/Plugins/Experimental/GeometryFlow/Source/GeometryFlowCore/Public/GeometryFlowGraph.h b/Engine/Plugins/Experimental/GeometryFlow/Source/GeometryFlowCore/Public/GeometryFlowGraph.h
--- a/Engine/Plugins/Experimental/GeometryFlow/Source/GeometryFlowCore/Public/GeometryFlowGraph.h
+++ b/Engine/Plugins/Experimental/GeometryFlow/Source/GeometryFlowCore/Public/GeometryFlowGraph.h
@@ -4,7 +4,7 @@
 
 #include "CoreMinimal.h"
 #include "GeometryFlowNode.h"
-
+#include "Templates/TypeHash.h"
 
 namespace UE
 {
@@ -42,6 +42,16 @@ public:
 		bool operator==(const FHandle& OtherHandle) const { return Identifier == OtherHandle.Identifier; }
 	};
 
+	friend uint32 GetTypeHash(FHandle Handle);
+
+	struct FConnection
+	{
+		FHandle FromNode;
+		FString FromOutput;
+		FHandle ToNode;
+		FString ToInput;
+	};
+
 	template<typename NodeType>
 	FHandle AddNodeOfType(
 		const FString& Identifier = FString(""),
@@ -52,7 +62,7 @@ public:
 		NewNodeInfo.Node->SetIdentifier(Identifier);
 		NewNodeInfo.CachingStrategy = CachingStrategy;
 		FHandle Handle = { NodeCounter++ };
-		AllNodes.Add(Handle.Identifier, NewNodeInfo);
+		AllNodes.Add(Handle, NewNodeInfo);
 		return Handle;
 	}
 
@@ -60,6 +70,19 @@ public:
 
 	EGeometryFlowResult InferConnection(FHandle FromNode, FHandle ToNode);
 
+	inline TSet<FHandle> GetSourceNodes() const
+	{
+		TSet<FHandle> SourceNodes;
+		for (const TMap<FHandle, FNodeInfo>::ElementType& IdNodePair : AllNodes)
+		{
+			if (IdNodePair.Value.Node->NodeInputs.Num() == 0)
+			{
+				SourceNodes.Add(FHandle{ IdNodePair.Key });
+			}
+		}
+		return SourceNodes;
+	}
+
 
 	template<typename T>
 	EGeometryFlowResult EvaluateResult(FHandle Node, FString OutputName, T& Storage, int32 StorageTypeIdentifier)
@@ -123,6 +146,9 @@ public:
 
 	
 protected:
+
+	friend class FGeometryFlowExecutor;
+
 	int32 NodeCounter = 0;
 
 	ENodeCachingStrategy DefaultCachingStrategy = ENodeCachingStrategy::AlwaysCache;
@@ -133,21 +159,13 @@ protected:
 		ENodeCachingStrategy CachingStrategy = ENodeCachingStrategy::Default;
 	};
 
-	TMap<int32, FNodeInfo> AllNodes;
+	TMap<FHandle, FNodeInfo> AllNodes;
 
 	TSafeSharedPtr<FNode> FindNode(FHandle Handle) const;
 	EGeometryFlowResult GetInputTypeForNode(FHandle NodeHandle, FString InputName, int32& Type) const;
 	EGeometryFlowResult GetOutputTypeForNode(FHandle NodeHandle, FString OutputName, int32& Type) const;
 	ENodeCachingStrategy GetCachingStrategyForNode(FHandle NodeHandle) const;
 
-	struct FConnection
-	{
-		FHandle FromNode;
-		FString FromOutput;
-
-		FHandle ToNode;
-		FString ToInput;
-	};
 	TArray<FConnection> Connections;
 
 	EGeometryFlowResult FindConnectionForInput(FHandle ToNode, FString ToInput, FConnection& ConnectionOut) const;
@@ -162,6 +180,10 @@ protected:
 };
 
 
+inline uint32 GetTypeHash(FGraph::FHandle Handle)
+{
+	return ::GetTypeHash(Handle.Identifier);
+}
 
 
 
@@ -172,4 +194,4 @@ protected:
 
 
 }	// end namespace GeometryFlow
-}	//
\ No newline at end of file
+}	//
diff --git a/Engine/Plugins/Experimental/GeometryFlow/Source/GeometryFlowCore/Public/GeometryFlowGraphUtil.h b/Engine/Plugins/Experimental/GeometryFlow/Source/GeometryFlowCore/Public/GeometryFlowGraphUtil.h
--- a/Engine/Plugins/Experimental/GeometryFlow/Source/GeometryFlowCore/Public/GeometryFlowGraphUtil.h
+++ b/Engine/Plugins/Experimental/GeometryFlow/Source/GeometryFlowCore/Public/GeometryFlowGraphUtil.h
@@ -35,6 +35,48 @@ void UpdateSettingsSourceNodeValue(FGraph& Graph, FGraph::FHandle NodeHandle, co
 	});
 }
 
+// Returns an index into the Connections array or -1 if not found
+inline int FindAnyConnectionFromNode(FGraph::FHandle FromNode, const TArray<FGraph::FConnection>& Connections)
+{
+	for (int ConnectionIndex = 0; ConnectionIndex < Connections.Num(); ++ConnectionIndex)
+	{
+		if (Connections[ConnectionIndex].FromNode == FromNode)
+		{
+			return ConnectionIndex;
+		}
+	}
+
+	return -1;
+}
+
+// Returns an index into the Connections array or -1 if not found
+inline int FindAnyConnectionToNode(FGraph::FHandle ToNode, const TArray<FGraph::FConnection>& Connections)
+{
+	for (int ConnectionIndex = 0; ConnectionIndex < Connections.Num(); ++ConnectionIndex)
+	{
+		if (Connections[ConnectionIndex].ToNode == ToNode)
+		{
+			return ConnectionIndex;
+		}
+	}
+
+	return -1;
+}
+
+// Returns an indices into the Connections array
+inline TArray<int> FindAllConnectionsToNode(FGraph::FHandle ToNodeID, const TArray<FGraph::FConnection>& Connections)
+{
+	TArray<int> Found;
+	for (int ConnectionIndex = 0; ConnectionIndex < Connections.Num(); ++ConnectionIndex)
+	{
+		if (Connections[ConnectionIndex].ToNode == ToNodeID)
+		{
+			Found.Add(ConnectionIndex);
+		}
+	}
+
+	return Found;
+}
 
 
 
diff --git a/Engine/Plugins/Experimental/GeometryFlow/Source/GeometryFlowCore/Public/GeometryFlowNode.h b/Engine/Plugins/Experimental/GeometryFlow/Source/GeometryFlowCore/Public/GeometryFlowNode.h
--- a/Engine/Plugins/Experimental/GeometryFlow/Source/GeometryFlowCore/Public/GeometryFlowNode.h
+++ b/Engine/Plugins/Experimental/GeometryFlow/Source/GeometryFlowCore/Public/GeometryFlowNode.h
@@ -305,6 +305,8 @@ public:
 
 	virtual TSafeSharedPtr<IData> StealOutput(FString OutputName);
 
+	TSafeSharedPtr<IData> GetOutput(const FString& OutputName) const;
+
 	/**
 	 * Find the list of named inputs that must be available to compute the named Outputs.
 	 * By default will return all Inputs if any of the listed Outputs exists on this Node.
@@ -313,6 +315,7 @@ public:
 		const TArray<FString>& Outputs, 
 		TArray<FEvalRequirement>& RequiredInputsOut);
 
+	virtual void CollectAllRequirements(TArray<FEvalRequirement>& RequiredInputsOut);
 
 	virtual void Evaluate(
 		const FNamedDataMap& DatasIn,
@@ -320,6 +323,10 @@ public:
 		TUniquePtr<FEvaluationInfo>& EvaluationInfo) = 0;
 
 protected:
+
+	friend class FGraph;
+	friend class FGeometryFlowExecutor;
+
 	FString Identifier;
 
 	struct FNodeInputInfo
@@ -353,7 +360,6 @@ protected:
 		const FString& OutputName,
 		TSafeSharedPtr<IData> NewData
 	);
-	TSafeSharedPtr<IData> GetOutput(const FString& OutputName) const;
 
 	TArray<FNodeInputInfo> NodeInputs;
 	TArray<FNodeOutputInfo> NodeOutputs;
diff --git a/Engine/Plugins/Experimental/GeometryFlow/Source/GeometryFlowCore/Public/GeometryFlowNodeUtil.h b/Engine/Plugins/Experimental/GeometryFlow/Source/GeometryFlowCore/Public/GeometryFlowNodeUtil.h
--- a/Engine/Plugins/Experimental/GeometryFlow/Source/GeometryFlowCore/Public/GeometryFlowNodeUtil.h
+++ b/Engine/Plugins/Experimental/GeometryFlow/Source/GeometryFlowCore/Public/GeometryFlowNodeUtil.h
@@ -31,8 +31,28 @@ TUniquePtr<TBasicNodeOutput<DataType, DataType::DataTypeIdentifier>> MakeBasicOu
 	return MakeUnique<TBasicNodeOutput<DataType, DataType::DataTypeIdentifier>>();
 }
 
+template<typename T>
+EGeometryFlowResult ExtractData(TSafeSharedPtr<IData> Data,
+								T& Storage,
+								int32 StorageTypeIdentifier,
+								bool bTryTakeResult)
+{
+	if (Data->GetPayloadType() != StorageTypeIdentifier)
+	{
+		return EGeometryFlowResult::UnmatchedTypes;
+	}
+	if (bTryTakeResult)
+	{
+		Data->GiveTo(Storage, StorageTypeIdentifier);
+	}
+	else
+	{
+		Data->GetDataCopy(Storage, StorageTypeIdentifier);
+	}
+
+	return EGeometryFlowResult::Ok;
 
-
+}
 
 
 
diff --git a/Engine/Plugins/Experimental/MeshLODToolset/Source/MeshLODToolset/Private/Graphs/GenerateMeshLODGraph.cpp b/Engine/Plugins/Experimental/MeshLODToolset/Source/MeshLODToolset/Private/Graphs/GenerateMeshLODGraph.cpp
--- a/Engine/Plugins/Experimental/MeshLODToolset/Source/MeshLODToolset/Private/Graphs/GenerateMeshLODGraph.cpp
+++ b/Engine/Plugins/Experimental/MeshLODToolset/Source/MeshLODToolset/Private/Graphs/GenerateMeshLODGraph.cpp
@@ -25,7 +25,7 @@
 
 #include "MeshDecompositionNodes/MakeTriangleSetsNode.h"
 #include "PhysicsNodes/GenerateConvexHullsCollisionNode.h"
-
+#include "GeometryFlowExecutor.h"
 
 using namespace UE::GeometryFlow;
 
@@ -36,6 +36,78 @@ void FGenerateMeshLODGraph::SetSourceMesh(const FDynamicMesh3& SourceMeshIn)
 	UpdateSourceNodeValue<FDynamicMeshSourceNode>(*Graph, MeshSourceNode, SourceMeshIn);
 }
 
+void FGenerateMeshLODGraph::EvaluateResultParallel(
+	FDynamicMesh3& ResultMesh,
+	FMeshTangentsd& ResultTangents,
+	FSimpleShapeSet3d& ResultCollision,
+	UE::GeometryFlow::FNormalMapImage& NormalMap,
+	TArray<TUniquePtr<UE::GeometryFlow::FTextureImage>>& TextureImages)
+{
+	//FScopedDurationTimeLogger Timer(TEXT("FGenerateMeshLODGraph::EvaluateResult -- parallel execution"));
+
+	TArray<FGeometryFlowExecutor::NodeOutputSpec> DesiredOutputs;
+	DesiredOutputs.Add(FGeometryFlowExecutor::NodeOutputSpec{ BakeNormalMapNode, FBakeMeshNormalMapNode::OutParamNormalMap() });
+	for (FBakeTextureGraphInfo& TexBakeStep : BakeTextureNodes)
+	{
+		DesiredOutputs.Add(FGeometryFlowExecutor::NodeOutputSpec{
+			TexBakeStep.BakeNode,
+			FBakeMeshTextureImageNode::OutParamTextureImage() });
+	}
+	DesiredOutputs.Add(FGeometryFlowExecutor::NodeOutputSpec{
+			CollisionOutputNode,
+			FCollisionGeometryTransferNode::OutParamValue() });
+	DesiredOutputs.Add(FGeometryFlowExecutor::NodeOutputSpec{
+		TangentsOutputNode,
+		FMeshTangentsTransferNode::OutParamValue() });
+	DesiredOutputs.Add(FGeometryFlowExecutor::NodeOutputSpec{
+		MeshOutputNode,
+		FDynamicMeshTransferNode::OutParamValue() });
+
+	FGeometryFlowExecutor Exec(*Graph);
+
+	TArray<TSafeSharedPtr<IData>> OutputDatas;
+	Exec.ComputeOutputs(DesiredOutputs, OutputDatas);
+
+	check(OutputDatas.Num() == DesiredOutputs.Num());
+
+	TArray<TSafeSharedPtr<IData>>::TIterator OutputDataIter = OutputDatas.CreateIterator();
+
+	// { BakeNormalMapNode, FBakeMeshNormalMapNode::OutParamNormalMap() }
+	NormalMap = FNormalMapImage();
+	bool bTakeNormalMap = true;
+	ExtractData(*OutputDataIter, NormalMap, (int)EMeshProcessingDataTypes::NormalMapImage, bTakeNormalMap);
+	++OutputDataIter;
+
+	for (FBakeTextureGraphInfo& TexBakeStep : BakeTextureNodes)
+	{
+		// { TexBakeStep.BakeNode, FBakeMeshTextureImageNode::OutParamTextureImage() }
+		TUniquePtr<UE::GeometryFlow::FTextureImage> NewImage = MakeUnique<UE::GeometryFlow::FTextureImage>();
+		bool bTakeTexBake = true;
+		ExtractData(*OutputDataIter, *NewImage, (int)EMeshProcessingDataTypes::TextureImage, bTakeTexBake);
+		++OutputDataIter;
+		TextureImages.Add(MoveTemp(NewImage));
+	}
+
+	// {CollisionOutputNode, FCollisionGeometryTransferNode::OutParamValue() }
+	ResultCollision = FSimpleShapeSet3d();
+	bool bTakeResultCollision = false;
+	ExtractData(*OutputDataIter, ResultCollision, (int)FCollisionGeometry::DataTypeIdentifier, bTakeResultCollision);
+	++OutputDataIter;
+
+	// { TangentsOutputNode, FMeshTangentsTransferNode::OutParamValue() }
+	bool bTakeResultTangents = false;
+	ResultTangents = FMeshTangentsd();
+	ExtractData(*OutputDataIter, ResultTangents, (int)EMeshProcessingDataTypes::MeshTangentSet, bTakeResultTangents);
+	++OutputDataIter;
+
+	// {MeshOutputNode, FDynamicMeshTransferNode::OutParamValue() }
+	bool bTakeResultMesh = true;
+	ResultMesh.Clear();
+	ExtractData(*OutputDataIter, ResultMesh, (int32)EMeshProcessingDataTypes::DynamicMesh, bTakeResultMesh);
+
+	++OutputDataIter;
+}
+
 
 void FGenerateMeshLODGraph::EvaluateResult(
 	FDynamicMesh3& ResultMesh,
@@ -44,6 +116,8 @@ void FGenerateMeshLODGraph::EvaluateResult(
 	UE::GeometryFlow::FNormalMapImage& NormalMap,
 	TArray<TUniquePtr<UE::GeometryFlow::FTextureImage>>& TextureImages)
 {
+	//FScopedDurationTimeLogger Timer(TEXT("FGenerateMeshLODGraph::EvaluateResult -- serial execution"));
+
 	//
 	// evaluate normal map
 	//
diff --git a/Engine/Plugins/Experimental/MeshLODToolset/Source/MeshLODToolset/Private/Graphs/GenerateStaticMeshLODProcess.cpp b/Engine/Plugins/Experimental/MeshLODToolset/Source/MeshLODToolset/Private/Graphs/GenerateStaticMeshLODProcess.cpp
--- a/Engine/Plugins/Experimental/MeshLODToolset/Source/MeshLODToolset/Private/Graphs/GenerateStaticMeshLODProcess.cpp
+++ b/Engine/Plugins/Experimental/MeshLODToolset/Source/MeshLODToolset/Private/Graphs/GenerateStaticMeshLODProcess.cpp
@@ -198,12 +198,24 @@ bool FGenerateStaticMeshLODProcess::InitializeGenerator()
 
 bool FGenerateStaticMeshLODProcess::ComputeDerivedSourceData()
 {
-	Generator->EvaluateResult(
-		this->DerivedLODMesh, 
-		this->DerivedLODMeshTangents,
-		this->DerivedCollision, 
-		this->DerivedNormalMapImage, 
-		this->DerivedTextureImages);
+	if (bUseParallelExecutor)
+	{
+		Generator->EvaluateResultParallel(
+			this->DerivedLODMesh,
+			this->DerivedLODMeshTangents,
+			this->DerivedCollision,
+			this->DerivedNormalMapImage,
+			this->DerivedTextureImages);
+	}
+	else
+	{
+		Generator->EvaluateResult(
+			this->DerivedLODMesh,
+			this->DerivedLODMeshTangents,
+			this->DerivedCollision,
+			this->DerivedNormalMapImage,
+			this->DerivedTextureImages);
+	}
 
 	// copy all materials for now...we are going to replace all the images though, and
 	// should not copy those?
diff --git a/Engine/Plugins/Experimental/MeshLODToolset/Source/MeshLODToolset/Private/Tools/GenerateStaticMeshLODAssetTool.cpp b/Engine/Plugins/Experimental/MeshLODToolset/Source/MeshLODToolset/Private/Tools/GenerateStaticMeshLODAssetTool.cpp
--- a/Engine/Plugins/Experimental/MeshLODToolset/Source/MeshLODToolset/Private/Tools/GenerateStaticMeshLODAssetTool.cpp
+++ b/Engine/Plugins/Experimental/MeshLODToolset/Source/MeshLODToolset/Private/Tools/GenerateStaticMeshLODAssetTool.cpp
@@ -129,6 +129,26 @@ void UGenerateStaticMeshLODAssetTool::Setup()
 	PreviewMesh->SetVisible(true);
 	PreviewMesh->SetTangentsMode(EDynamicMeshTangentCalcType::ExternallyCalculated);
 
+	// Recompute if we switch between parallel and serial
+	int32 WatcherIndex = BasicProperties->WatchProperty(BasicProperties->bParallelExecution, [this](bool bNewParallelExec)
+	{
+		// TODO: We crash if we don't recreate the Process and reinitialize it. Why?
+
+		GenerateProcess = MakePimpl<FGenerateStaticMeshLODProcess>();
+
+		TUniquePtr<FPrimitiveComponentTarget>& SourceComponent = ComponentTargets[0];
+		UStaticMeshComponent* StaticMeshComponent = CastChecked<UStaticMeshComponent>(SourceComponent->GetOwnerComponent());
+		UStaticMesh* StaticMesh = StaticMeshComponent->GetStaticMesh();
+		if (StaticMesh)
+		{
+			GenerateProcess->Initialize(StaticMesh);
+		}
+
+		bPreviewValid = false;
+		ValidatePreview();
+	});
+	BasicProperties->SilentUpdateWatcherAtIndex(WatcherIndex);
+
 	bPreviewValid = false;
 	ValidatePreview();
 }
@@ -176,6 +196,8 @@ void UGenerateStaticMeshLODAssetTool::ValidatePreview()
 {
 	if (bPreviewValid) return;
 
+	GenerateProcess->bUseParallelExecutor = BasicProperties->bParallelExecution;
+
 	GenerateProcess->ComputeDerivedSourceData();
 	const FDynamicMesh3& ResultMesh = GenerateProcess->GetDerivedLOD0Mesh();
 	const FMeshTangentsd& ResultTangents = GenerateProcess->GetDerivedLOD0MeshTangents();
diff --git a/Engine/Plugins/Experimental/MeshLODToolset/Source/MeshLODToolset/Public/Graphs/GenerateMeshLODGraph.h b/Engine/Plugins/Experimental/MeshLODToolset/Source/MeshLODToolset/Public/Graphs/GenerateMeshLODGraph.h
--- a/Engine/Plugins/Experimental/MeshLODToolset/Source/MeshLODToolset/Public/Graphs/GenerateMeshLODGraph.h
+++ b/Engine/Plugins/Experimental/MeshLODToolset/Source/MeshLODToolset/Public/Graphs/GenerateMeshLODGraph.h
@@ -27,6 +27,11 @@ public:
 		UE::GeometryFlow::FNormalMapImage& NormalMap,
 		TArray<TUniquePtr<UE::GeometryFlow::FTextureImage>>& TextureImages);
 
+	void EvaluateResultParallel(FDynamicMesh3& ResultMesh,
+						   FMeshTangentsd& ResultTangents,
+						   FSimpleShapeSet3d& ResultCollision,
+						   UE::GeometryFlow::FNormalMapImage& NormalMap,
+						   TArray<TUniquePtr<UE::GeometryFlow::FTextureImage>>& TextureImages);
 
 protected:
 	TUniquePtr<UE::GeometryFlow::FGraph> Graph;
diff --git a/Engine/Plugins/Experimental/MeshLODToolset/Source/MeshLODToolset/Public/Graphs/GenerateStaticMeshLODProcess.h b/Engine/Plugins/Experimental/MeshLODToolset/Source/MeshLODToolset/Public/Graphs/GenerateStaticMeshLODProcess.h
--- a/Engine/Plugins/Experimental/MeshLODToolset/Source/MeshLODToolset/Public/Graphs/GenerateStaticMeshLODProcess.h
+++ b/Engine/Plugins/Experimental/MeshLODToolset/Source/MeshLODToolset/Public/Graphs/GenerateStaticMeshLODProcess.h
@@ -55,6 +55,7 @@ public:
 	};
 	void GetDerivedMaterialsPreview(FPreviewMaterials& MaterialSetOut);
 
+	bool bUseParallelExecutor = false;
 
 protected:
 
diff --git a/Engine/Plugins/Experimental/MeshLODToolset/Source/MeshLODToolset/Public/Tools/GenerateStaticMeshLODAssetTool.h b/Engine/Plugins/Experimental/MeshLODToolset/Source/MeshLODToolset/Public/Tools/GenerateStaticMeshLODAssetTool.h
--- a/Engine/Plugins/Experimental/MeshLODToolset/Source/MeshLODToolset/Public/Tools/GenerateStaticMeshLODAssetTool.h
+++ b/Engine/Plugins/Experimental/MeshLODToolset/Source/MeshLODToolset/Public/Tools/GenerateStaticMeshLODAssetTool.h
@@ -64,6 +64,10 @@ public:
 
 	UPROPERTY(VisibleAnywhere, Category = Previews)
 	TArray<UTexture2D*> PreviewTextures;
+
+	UPROPERTY(EditAnywhere, Category = General)
+	bool bParallelExecution = false;
+
 };
 
 
