GeometryProcessingUnitTests: default to multithreaded mode, use --no-mt for single threaded. Detect and skip tests that would crash with no multithreading.

#rnx

diff --git a/Engine/Plugins/Experimental/GeometryFlow/Source/GeometryFlowCore/Private/GeometryFlowExecutor.cpp b/Engine/Plugins/Experimental/GeometryFlow/Source/GeometryFlowCore/Private/GeometryFlowExecutor.cpp
--- a/Engine/Plugins/Experimental/GeometryFlow/Source/GeometryFlowCore/Private/GeometryFlowExecutor.cpp
+++ b/Engine/Plugins/Experimental/GeometryFlow/Source/GeometryFlowCore/Private/GeometryFlowExecutor.cpp
@@ -31,8 +31,8 @@ FGeometryFlowExecutor::FGeometryFlowExecutor(FGraph& InGraph)
 	EvalInfo = MakeUnique<FEvaluationInfo>();
 	TopologicalSort();
 
-	// Check that TaskGraph is enabled!
-	check(FTaskGraphInterface::IsRunning());
+	// Check that TaskGraph is enabled
+	check(FTaskGraphInterface::IsRunning() && FPlatformProcess::SupportsMultithreading());
 }
 
 void FGeometryFlowExecutor::TopologicalSort()
