Geometry flow: add a normal flow remesher node

#rnx

diff --git a/Engine/Plugins/Experimental/GeometryFlow/Source/GeometryFlowMeshProcessing/Private/MeshProcessingNodes/MeshNormalFlowNode.cpp b/Engine/Plugins/Experimental/GeometryFlow/Source/GeometryFlowMeshProcessing/Private/MeshProcessingNodes/MeshNormalFlowNode.cpp
new file mode 100644
--- /dev/null
+++ b/Engine/Plugins/Experimental/GeometryFlow/Source/GeometryFlowMeshProcessing/Private/MeshProcessingNodes/MeshNormalFlowNode.cpp
@@ -0,0 +1,124 @@
+// Copyright Epic Games, Inc. All Rights Reserved.
+
+#include "MeshProcessingNodes/MeshNormalFlowNode.h"
+#include "NormalFlowRemesher.h"
+#include "ProjectionTargets.h"
+#include "MeshNormals.h"
+
+using namespace UE::GeometryFlow;
+
+namespace
+{
+	// TODO: This duplicates code in a private function in URemeshMeshTool
+	double CalculateTargetEdgeLength(const FDynamicMesh3& Mesh, int TargetTriCount)
+	{
+		double InitialMeshArea = 0;
+		for (int TriangleIndex : Mesh.TriangleIndicesItr())
+		{
+			InitialMeshArea += Mesh.GetTriArea(TriangleIndex);
+		}
+
+		double TargetTriArea = InitialMeshArea / (double)TargetTriCount;
+		double EdgeLen = TriangleUtil::EquilateralEdgeLengthForArea(TargetTriArea);
+		return (double)FMath::RoundToInt(EdgeLen * 100.0) / 100.0;
+	}
+}
+
+void FMeshNormalFlowNode::CheckAdditionalInputs(const FNamedDataMap& DatasIn, bool& bRecomputeRequired, bool& bAllInputsValid)
+{
+	FindAndUpdateInputForEvaluate(InParamTargetMesh(), DatasIn, bRecomputeRequired, bAllInputsValid);
+}
+
+
+void FMeshNormalFlowNode::DoNormalFlow(const FMeshNormalFlowSettings& SettingsIn,
+									   const FDynamicMesh3& ProjectionTargetMesh,
+									   bool bAttributesHaveBeenDiscarded,
+									   FDynamicMesh3& EditMesh)
+{
+	FNormalFlowRemesher Remesher(&EditMesh);
+
+	Remesher.MaxRemeshIterations = SettingsIn.MaxRemeshIterations;
+	Remesher.NumExtraProjectionIterations = SettingsIn.NumExtraProjectionIterations;
+
+	Remesher.bEnableSplits = SettingsIn.bSplits;
+	Remesher.bEnableFlips = SettingsIn.bFlips;
+	Remesher.bEnableCollapses = SettingsIn.bCollapses;
+
+	double TargetEdgeLength = CalculateTargetEdgeLength(EditMesh, SettingsIn.TargetCount);
+	Remesher.SetTargetEdgeLength(TargetEdgeLength);
+
+	Remesher.bEnableSmoothing = (SettingsIn.SmoothingStrength > 0);
+	Remesher.SmoothSpeedT = SettingsIn.SmoothingStrength;
+	Remesher.SmoothType = SettingsIn.SmoothingType;
+
+	if (!SettingsIn.bDiscardAttributes)
+	{
+		Remesher.FlipMetric = (SettingsIn.SmoothingType == FRemesher::ESmoothTypes::Uniform) ?
+			FRemesher::EFlipMetric::OptimalValence : FRemesher::EFlipMetric::MinEdgeLength;
+	}
+
+	Remesher.bPreventNormalFlips = SettingsIn.bPreventNormalFlips;
+	Remesher.DEBUG_CHECK_LEVEL = 0;
+
+	FMeshConstraints Constraints;
+	FMeshConstraintsUtil::ConstrainAllBoundariesAndSeams(Constraints, EditMesh,
+														 SettingsIn.MeshBoundaryConstraints,
+														 SettingsIn.GroupBorderConstraints,
+														 SettingsIn.MaterialBorderConstraints,
+														 true, !SettingsIn.bPreserveSharpEdges);
+
+	Remesher.SetExternalConstraints(MoveTemp(Constraints));
+
+
+	FDynamicMeshAABBTree3 ProjectionTargetSpatial(&ProjectionTargetMesh, true);
+	FMeshProjectionTarget ProjTarget(&ProjectionTargetMesh, &ProjectionTargetSpatial);
+	Remesher.SetProjectionTarget(&ProjTarget);
+
+	if (SettingsIn.bDiscardAttributes && !bAttributesHaveBeenDiscarded)
+	{
+		EditMesh.DiscardAttributes();
+	}
+
+	// Go
+	Remesher.BasicRemeshPass();
+
+	if (!EditMesh.HasAttributes())
+	{
+		FMeshNormals::QuickComputeVertexNormals(EditMesh);
+	}
+	else
+	{
+		FMeshNormals::QuickRecomputeOverlayNormals(EditMesh);
+	}
+}
+
+void FMeshNormalFlowNode::ProcessMesh(
+	const FNamedDataMap& DatasIn,
+	const FMeshNormalFlowSettings& SettingsIn,
+	const FDynamicMesh3& MeshIn,
+	FDynamicMesh3& MeshOut)
+{
+	TSafeSharedPtr<IData> ProjectionTargetMeshArg = DatasIn.FindData(InParamTargetMesh());
+	FDynamicMesh3 ProjectionTargetMesh;
+	ProjectionTargetMeshArg->GetDataCopy<FDynamicMesh3>(ProjectionTargetMesh, (int)EMeshProcessingDataTypes::DynamicMesh);
+
+	bool bDiscardAttributesImmediately = SettingsIn.bDiscardAttributes && !SettingsIn.bPreserveSharpEdges;
+	MeshOut.Copy(MeshIn, true, true, true, !bDiscardAttributesImmediately);
+
+	DoNormalFlow(SettingsIn, ProjectionTargetMesh, bDiscardAttributesImmediately, MeshOut);
+}
+
+
+void FMeshNormalFlowNode::ProcessMeshInPlace(
+	const FNamedDataMap& DatasIn,
+	const FMeshNormalFlowSettings& SettingsIn,
+	FDynamicMesh3& MeshInOut)
+{
+	// TODO: In this case, can we call IData::GiveTo instead of GetDataCopy for the projection target?
+	TSafeSharedPtr<IData> ProjectionTargetMeshArg = DatasIn.FindData(InParamTargetMesh());
+	FDynamicMesh3 ProjectionTargetMesh;
+	ProjectionTargetMeshArg->GetDataCopy<FDynamicMesh3>(ProjectionTargetMesh, (int)EMeshProcessingDataTypes::DynamicMesh);
+
+	constexpr bool bAttributesHaveBeenDiscarded = false;
+	DoNormalFlow(SettingsIn, ProjectionTargetMesh, bAttributesHaveBeenDiscarded, MeshInOut);
+}
diff --git a/Engine/Plugins/Experimental/GeometryFlow/Source/GeometryFlowMeshProcessing/Public/MeshProcessingNodes/MeshNormalFlowNode.h b/Engine/Plugins/Experimental/GeometryFlow/Source/GeometryFlowMeshProcessing/Public/MeshProcessingNodes/MeshNormalFlowNode.h
new file mode 100644
--- /dev/null
+++ b/Engine/Plugins/Experimental/GeometryFlow/Source/GeometryFlowMeshProcessing/Public/MeshProcessingNodes/MeshNormalFlowNode.h
@@ -0,0 +1,70 @@
+// Copyright Epic Games, Inc. All Rights Reserved.
+
+#pragma once
+
+#include "MeshProcessingNodes/MeshProcessingBaseNodes.h"
+#include "MeshProcessingNodes/MeshProcessingDataTypes.h"
+#include "MeshProcessingNodes/MeshSimplifyNode.h"
+#include "Remesher.h"
+
+namespace UE
+{
+namespace GeometryFlow
+{
+
+struct FMeshNormalFlowSettings : public FMeshSimplifySettings
+{
+	DECLARE_GEOMETRYFLOW_DATA_TYPE_IDENTIFIER(EMeshProcessingDataTypes::NormalFlowSettings);
+
+	int MaxRemeshIterations = 20;
+	int NumExtraProjectionIterations = 5;
+	bool bFlips = true;
+	bool bSplits = true;
+	bool bCollapses = true;
+	FRemesher::ESmoothTypes SmoothingType = FRemesher::ESmoothTypes::Uniform;
+	float SmoothingStrength = 0.25f;
+};
+
+GEOMETRYFLOW_DECLARE_SETTINGS_TYPES(FMeshNormalFlowSettings, MeshNormalFlow);
+
+
+class FMeshNormalFlowNode : public TProcessMeshWithSettingsBaseNode<FMeshNormalFlowSettings>
+{
+public:
+
+	// Beside the intput mesh and settings, we also take a "target mesh" that will be the projection target.
+	static const FString InParamTargetMesh() { return TEXT("TargetMesh"); }
+
+	FMeshNormalFlowNode() : TProcessMeshWithSettingsBaseNode<FMeshNormalFlowSettings>()
+	{
+		AddInput(InParamTargetMesh(), MakeUnique<FDynamicMeshInput>());
+	}
+
+	void ProcessMesh(
+		const FNamedDataMap& DatasIn,
+		const FMeshNormalFlowSettings& SettingsIn,
+		const FDynamicMesh3& MeshIn,
+		FDynamicMesh3& MeshOut) override;
+
+	void ProcessMeshInPlace(
+		const FNamedDataMap& DatasIn,
+		const FMeshNormalFlowSettings& SettingsIn,
+		FDynamicMesh3& MeshInOut) override;
+
+	void CheckAdditionalInputs(const FNamedDataMap& DatasIn, 
+							   bool& bRecomputeRequired, 
+							   bool& bAllInputsValid) override;
+
+
+protected:
+
+	void DoNormalFlow(const FMeshNormalFlowSettings& SettingsIn,
+					  const FDynamicMesh3& ProjectionTargetMesh,
+					  bool bAttributesHaveBeenDiscarded,
+					  FDynamicMesh3& EditMesh);
+
+};
+
+
+}	// end namespace GeometryFlow
+}	// end namespace UE
diff --git a/Engine/Plugins/Experimental/GeometryFlow/Source/GeometryFlowMeshProcessing/Public/MeshProcessingNodes/MeshProcessingDataTypes.h b/Engine/Plugins/Experimental/GeometryFlow/Source/GeometryFlowMeshProcessing/Public/MeshProcessingNodes/MeshProcessingDataTypes.h
--- a/Engine/Plugins/Experimental/GeometryFlow/Source/GeometryFlowMeshProcessing/Public/MeshProcessingNodes/MeshProcessingDataTypes.h
+++ b/Engine/Plugins/Experimental/GeometryFlow/Source/GeometryFlowMeshProcessing/Public/MeshProcessingNodes/MeshProcessingDataTypes.h
@@ -27,6 +27,7 @@ enum class EMeshProcessingDataTypes
 	VoxMorphologyOpSettings = (int)EDataTypes::BaseMeshProcessingTypes + 101,
 
 	SimplifySettings = (int)EDataTypes::BaseMeshProcessingTypes + 110,
+	NormalFlowSettings = (int)EDataTypes::BaseMeshProcessingTypes + 111,
 
 
 
