Fix static analysis by adding ".0f" to a literal.

diff --git a/Engine/Plugins/Experimental/GeometryFlow/Source/GeometryFlowMeshProcessing/Private/MeshProcessingNodes/MeshNormalFlowNode.cpp b/Engine/Plugins/Experimental/GeometryFlow/Source/GeometryFlowMeshProcessing/Private/MeshProcessingNodes/MeshNormalFlowNode.cpp
--- a/Engine/Plugins/Experimental/GeometryFlow/Source/GeometryFlowMeshProcessing/Private/MeshProcessingNodes/MeshNormalFlowNode.cpp
+++ b/Engine/Plugins/Experimental/GeometryFlow/Source/GeometryFlowMeshProcessing/Private/MeshProcessingNodes/MeshNormalFlowNode.cpp
@@ -47,7 +47,7 @@ void FMeshNormalFlowNode::DoNormalFlow(const FMeshNormalFlowSettings& SettingsIn
 	double TargetEdgeLength = CalculateTargetEdgeLength(EditMesh, SettingsIn.TargetCount);
 	Remesher.SetTargetEdgeLength(TargetEdgeLength);
 
-	Remesher.bEnableSmoothing = (SettingsIn.SmoothingStrength > 0);
+	Remesher.bEnableSmoothing = (SettingsIn.SmoothingStrength > 0.0f);
 	Remesher.SmoothSpeedT = SettingsIn.SmoothingStrength;
 	Remesher.SmoothType = SettingsIn.SmoothingType;
 
