Add debug dump of graph connectivity (node-to-node only, not input-to-output). Allows for graph visualization in third party applications.

#rnx

diff --git a/Engine/Plugins/Experimental/GeometryFlow/Source/GeometryFlowCore/Private/GeometryFlowGraph.cpp b/Engine/Plugins/Experimental/GeometryFlow/Source/GeometryFlowCore/Private/GeometryFlowGraph.cpp
--- a/Engine/Plugins/Experimental/GeometryFlow/Source/GeometryFlowCore/Private/GeometryFlowGraph.cpp
+++ b/Engine/Plugins/Experimental/GeometryFlow/Source/GeometryFlowCore/Private/GeometryFlowGraph.cpp
@@ -241,4 +241,61 @@ EGeometryFlowResult FGraph::SetNodeCachingStrategy(FHandle NodeHandle, ENodeCach
 	}
 	Found->CachingStrategy = Strategy;
 	return EGeometryFlowResult::Ok;
-}
\ No newline at end of file
+}
+
+
+FString FGraph::DebugDumpGraph(TFunction<bool(TSafeSharedPtr<FNode>)> IncludeNodeFn) const
+{
+	// Can be used by, e.g., https://csacademy.com/app/graph_editor/
+
+	FString Out;
+
+	// First, all node names
+	for (const TPair<FHandle, FNodeInfo>& NodeHandleAndInfo : AllNodes)
+	{
+		if (!NodeHandleAndInfo.Value.Node)
+		{ 
+			return "Error"; 
+		}
+			
+		if (!IncludeNodeFn(NodeHandleAndInfo.Value.Node))
+		{ 
+			continue; 
+		}		
+
+		FString NodeName = NodeHandleAndInfo.Value.Node->GetIdentifier();
+		Out += NodeName + "\n";
+	}
+
+	// Second, connections by node name
+	for (const FConnection& Connection : Connections)
+	{
+		FHandle FromHandle = Connection.FromNode;
+		if (!AllNodes.Find(FromHandle) || !AllNodes.Find(FromHandle)->Node) 
+		{ 
+			return "Error"; 
+		}
+
+		TSafeSharedPtr<FNode> FromNode = AllNodes.Find(FromHandle)->Node;
+		if (!IncludeNodeFn(FromNode))
+		{
+			continue;
+		}
+
+		FHandle ToHandle = Connection.ToNode;
+		if (!AllNodes.Find(ToHandle) || !AllNodes.Find(ToHandle)->Node) 
+		{ 
+			return "Error"; 
+		}
+
+		TSafeSharedPtr<FNode> ToNode = AllNodes.Find(ToHandle)->Node;
+		if (!IncludeNodeFn(ToNode))
+		{ 
+			continue; 
+		}
+
+		Out += FromNode->GetIdentifier() + " " + ToNode->GetIdentifier() + "\n";
+	}
+
+	return Out;
+}
diff --git a/Engine/Plugins/Experimental/GeometryFlow/Source/GeometryFlowCore/Public/GeometryFlowGraph.h b/Engine/Plugins/Experimental/GeometryFlow/Source/GeometryFlowCore/Public/GeometryFlowGraph.h
--- a/Engine/Plugins/Experimental/GeometryFlow/Source/GeometryFlowCore/Public/GeometryFlowGraph.h
+++ b/Engine/Plugins/Experimental/GeometryFlow/Source/GeometryFlowCore/Public/GeometryFlowGraph.h
@@ -144,7 +144,8 @@ public:
 	void ConfigureCachingStrategy(ENodeCachingStrategy NewStrategy);
 	EGeometryFlowResult SetNodeCachingStrategy(FHandle NodeHandle, ENodeCachingStrategy Strategy);
 
-	
+	FString DebugDumpGraph(TFunction<bool(TSafeSharedPtr<FNode>)> IncludeNodeFn) const;
+
 protected:
 
 	friend class FGeometryFlowExecutor;
diff --git a/Engine/Plugins/Experimental/MeshLODToolset/Source/MeshLODToolset/Private/Graphs/GenerateMeshLODGraph.cpp b/Engine/Plugins/Experimental/MeshLODToolset/Source/MeshLODToolset/Private/Graphs/GenerateMeshLODGraph.cpp
--- a/Engine/Plugins/Experimental/MeshLODToolset/Source/MeshLODToolset/Private/Graphs/GenerateMeshLODGraph.cpp
+++ b/Engine/Plugins/Experimental/MeshLODToolset/Source/MeshLODToolset/Private/Graphs/GenerateMeshLODGraph.cpp
@@ -399,8 +399,6 @@ void FGenerateMeshLODGraph::BuildGraph()
 	Graph->InferConnection(BakeNormalMapSettingsNode, BakeNormalMapNode);
 
 
-	Graph->AddNodeOfType<FMeshDeleteTrianglesNode>(TEXT("TestDeleteTrisNode"));
-
 	// collision generation
 
 	FGraph::FHandle IgnoreGroupsForCollisionNode = Graph->AddNodeOfType<FIndexSetsSourceNode>(TEXT("CollisionIgnoreGroups"));
@@ -477,6 +475,14 @@ void FGenerateMeshLODGraph::BuildGraph()
 
 	FGenerateConvexHullsCollisionSettings GenConvexesSettings;
 	UpdateGenerateConvexCollisionSettings(GenConvexesSettings);
+
+
+	//FString GraphDump = Graph->DebugDumpGraph([](TSafeSharedPtr<FNode> Node)
+	//{
+	//	return !Node->GetIdentifier().EndsWith("Settings");
+	//});
+	//UE_LOG(LogTemp, Warning, TEXT("GRAPH:\n%s"), *GraphDump);
+
 }
 
 
