Add an option to pre-filter mesh vertices before computing convex hull. Result is not a convex hull in general, but might be faster to compute.

#rnx

diff --git a/Engine/Plugins/Experimental/GeometryFlow/Source/GeometryFlowMeshProcessing/Private/PhysicsNodes/GenerateConvexHullsCollisionNode.cpp b/Engine/Plugins/Experimental/GeometryFlow/Source/GeometryFlowMeshProcessing/Private/PhysicsNodes/GenerateConvexHullsCollisionNode.cpp
--- a/Engine/Plugins/Experimental/GeometryFlow/Source/GeometryFlowMeshProcessing/Private/PhysicsNodes/GenerateConvexHullsCollisionNode.cpp
+++ b/Engine/Plugins/Experimental/GeometryFlow/Source/GeometryFlowMeshProcessing/Private/PhysicsNodes/GenerateConvexHullsCollisionNode.cpp
@@ -6,7 +6,7 @@
 #include "Operations/MeshConvexHull.h"
 
 #include "Async/ParallelFor.h"
-
+#include "ProfilingDebugging/ScopedTimers.h"
 
 using namespace UE::GeometryFlow;
 
@@ -32,31 +32,51 @@ void FGenerateConvexHullsCollisionNode::Evaluate(
 
 				const FDynamicMesh3& Mesh = MeshArg->GetDataConstRef<FDynamicMesh3>((int)EMeshProcessingDataTypes::DynamicMesh);
 				const FIndexSets& IndexData = TriSetsArg->GetDataConstRef<FIndexSets>((int)EMeshProcessingDataTypes::IndexSets);
+				FCollisionGeometry Result;
 
-				int32 NumShapes = IndexData.IndexSets.Num();
-				TArray<FConvexShape3d> Convexes;
-				TArray<bool> ConvexOk;
-				Convexes.SetNum(NumShapes);
-				ConvexOk.SetNum(NumShapes);
-				ParallelFor(NumShapes, [&](int32 k)
+				if (Settings.bPrefilterVertices)
 				{
-					ConvexOk[k] = false;
+					// TODO: Combine prefiltering and mesh segmentation, they are not mutually exclusive
+
+					//FScopedDurationTimeLogger Time(TEXT(" @@@@@@@ Convex hull with prefiltering"));
+
 					FMeshConvexHull Hull(&Mesh);
-					MeshIndexUtil::TriangleToVertexIDs(&Mesh, IndexData.IndexSets[k], Hull.VertexSet);
+					FMeshConvexHull::GridSample(Mesh, Settings.PrefilterGridResolution, Hull.VertexSet);
 					Hull.bPostSimplify = Settings.SimplifyToTriangleCount > 0;
-					Hull.MaxTargetFaceCount = FMath::Max(1, Settings.SimplifyToTriangleCount);
+					Hull.MaxTargetFaceCount = Settings.SimplifyToTriangleCount;
 					if (Hull.Compute())
 					{
 						FConvexShape3d NewConvex;
 						NewConvex.Mesh = MoveTemp(Hull.ConvexHull);
-						Convexes[k] = MoveTemp(NewConvex);
-						ConvexOk[k] = true;
+						Result.Geometry.Convexes.Add(MoveTemp(NewConvex));
 					}
-				});
-
-				FCollisionGeometry Result;
-				Result.Geometry.Convexes = MoveTemp(Convexes);
+				}
+				else
+				{
+					//FScopedDurationTimeLogger Time(TEXT(" @@@@@@@ Convex hull with mesh segmentation"));
 
+					int32 NumShapes = IndexData.IndexSets.Num();
+					TArray<FConvexShape3d> Convexes;
+					TArray<bool> ConvexOk;
+					Convexes.SetNum(NumShapes);
+					ConvexOk.SetNum(NumShapes);
+					ParallelFor(NumShapes, [&](int32 k)
+					{
+						ConvexOk[k] = false;
+						FMeshConvexHull Hull(&Mesh);
+						MeshIndexUtil::TriangleToVertexIDs(&Mesh, IndexData.IndexSets[k], Hull.VertexSet);
+						Hull.bPostSimplify = Settings.SimplifyToTriangleCount > 0;
+						Hull.MaxTargetFaceCount = Settings.SimplifyToTriangleCount;
+						if (Hull.Compute())
+						{
+							FConvexShape3d NewConvex;
+							NewConvex.Mesh = MoveTemp(Hull.ConvexHull);
+							Convexes[k] = MoveTemp(NewConvex);
+							ConvexOk[k] = true;
+						}
+					});
+					Result.Geometry.Convexes = MoveTemp(Convexes);
+				}
 				SetOutput(OutParamGeometry(), MakeMovableData<FCollisionGeometry>(MoveTemp(Result)));
 				EvaluationInfo->CountCompute(this);
 			}
diff --git a/Engine/Plugins/Experimental/GeometryFlow/Source/GeometryFlowMeshProcessing/Public/PhysicsNodes/GenerateConvexHullsCollisionNode.h b/Engine/Plugins/Experimental/GeometryFlow/Source/GeometryFlowMeshProcessing/Public/PhysicsNodes/GenerateConvexHullsCollisionNode.h
--- a/Engine/Plugins/Experimental/GeometryFlow/Source/GeometryFlowMeshProcessing/Public/PhysicsNodes/GenerateConvexHullsCollisionNode.h
+++ b/Engine/Plugins/Experimental/GeometryFlow/Source/GeometryFlowMeshProcessing/Public/PhysicsNodes/GenerateConvexHullsCollisionNode.h
@@ -23,6 +23,8 @@ struct GEOMETRYFLOWMESHPROCESSING_API FGenerateConvexHullsCollisionSettings
 	DECLARE_GEOMETRYFLOW_DATA_TYPE_IDENTIFIER(EMeshProcessingDataTypes::GenerateCollisionConvexHullsSettings);
 
 	int32 SimplifyToTriangleCount = 50;
+	bool bPrefilterVertices = false;
+	int PrefilterGridResolution = 10;
 	
 };
 GEOMETRYFLOW_DECLARE_SETTINGS_TYPES(FGenerateConvexHullsCollisionSettings, GenerateConvexHullsCollision);
diff --git a/Engine/Plugins/Experimental/GeometryProcessing/Source/DynamicMesh/Private/Operations/MeshConvexHull.cpp b/Engine/Plugins/Experimental/GeometryProcessing/Source/DynamicMesh/Private/Operations/MeshConvexHull.cpp
--- a/Engine/Plugins/Experimental/GeometryProcessing/Source/DynamicMesh/Private/Operations/MeshConvexHull.cpp
+++ b/Engine/Plugins/Experimental/GeometryProcessing/Source/DynamicMesh/Private/Operations/MeshConvexHull.cpp
@@ -5,7 +5,7 @@
 #include "MeshSimplification.h"
 #include "MeshNormals.h"
 #include "ConvexHull3.h"
-
+#include "Util/GridIndexing3.h"
 
 bool FMeshConvexHull::Compute()
 {
@@ -26,6 +26,7 @@ bool FMeshConvexHull::Compute()
 
 	if (bPostSimplify)
 	{
+		check(MaxTargetFaceCount > 0);
 		bool bSimplified = false;
 		if (ConvexHull.TriangleCount() > MaxTargetFaceCount)
 		{
@@ -131,4 +132,63 @@ bool FMeshConvexHull::Compute_VertexSubset()
 	});
 
 	return true;
-}
\ No newline at end of file
+}
+
+
+FVector3i FMeshConvexHull::DebugGetCellIndex(const FDynamicMesh3& Mesh,
+											 int GridResolutionMaxAxis,
+											 int VertexIndex)
+{
+	FAxisAlignedBox3d Bounds = Mesh.GetBounds();
+	Bounds.Min = Bounds.Min - 1e-4;			// Pad to avoid problems with vertices lying exactly on bounding box
+	Bounds.Max = Bounds.Max + 1e-4;
+
+	const double GridCellSize = Bounds.MaxDim() / (double)GridResolutionMaxAxis;
+
+	FBoundsGridIndexer3d Indexer(Bounds, GridCellSize);
+
+	return Indexer.ToGrid(Mesh.GetVertex(VertexIndex));
+}
+
+
+void FMeshConvexHull::GridSample(const FDynamicMesh3& Mesh, 
+								 int GridResolutionMaxAxis, 
+								 TArray<int32>& OutSamples)
+{
+	// Simple spatial hash to find a representative vertex for each occupied grid cell
+
+	FAxisAlignedBox3d Bounds = Mesh.GetBounds();
+	Bounds.Min = Bounds.Min - 1e-4;			// Pad to avoid problems with vertices lying exactly on bounding box
+	Bounds.Max = Bounds.Max + 1e-4;
+	const double GridCellSize = Bounds.MaxDim() / (double)GridResolutionMaxAxis;
+
+	FBoundsGridIndexer3d Indexer(Bounds, GridCellSize);
+	const FVector3i GridResolution = Indexer.GridResolution();
+
+	// TODO: If the grid resolution is too high, use a TMap from grid cell index to vertex index instead of an array.
+	// For smallish grids the array is more efficient.
+	int TotalNumberGridCells = GridResolution.X * GridResolution.Y * GridResolution.Z;
+	TArray<int32> GridCellVertex;
+	GridCellVertex.Init(-1, TotalNumberGridCells);
+
+	for (int VertexIndex : Mesh.VertexIndicesItr())
+	{
+		FVector3i CellIndex = Indexer.ToGrid(Mesh.GetVertex(VertexIndex));
+		check(CellIndex.X >= 0 && CellIndex.X < GridResolution.X);
+		check(CellIndex.Y >= 0 && CellIndex.Y < GridResolution.Y);
+		check(CellIndex.Z >= 0 && CellIndex.Z < GridResolution.Z);
+
+		int Key = CellIndex.X + CellIndex.Y * GridResolution.X + CellIndex.Z * GridResolution.X * GridResolution.Y;
+
+		GridCellVertex[Key] = VertexIndex;
+	}
+
+	for (const int32 VertexIndex : GridCellVertex)
+	{
+		if (VertexIndex >= 0)
+		{
+			OutSamples.Add(VertexIndex);
+		}
+	}
+
+}
diff --git a/Engine/Plugins/Experimental/GeometryProcessing/Source/DynamicMesh/Public/Operations/MeshConvexHull.h b/Engine/Plugins/Experimental/GeometryProcessing/Source/DynamicMesh/Public/Operations/MeshConvexHull.h
--- a/Engine/Plugins/Experimental/GeometryProcessing/Source/DynamicMesh/Public/Operations/MeshConvexHull.h
+++ b/Engine/Plugins/Experimental/GeometryProcessing/Source/DynamicMesh/Public/Operations/MeshConvexHull.h
@@ -27,6 +27,24 @@ public:
 	/** Output convex hull */
 	FDynamicMesh3 ConvexHull;
 
+	/** 
+	 *  Choose a more or less evenly-spaced subset of mesh vertices. Conceptually, this function creates a uniform 
+	 *  grid with given cell size. (Cell size is given roughly as a percentage of the total mesh bounding box.) Each 
+	 *  grid cell can hold up to one vertex. Return the set of representative vertices, maximum one per cell.
+	 * 
+	 *  @param Mesh							Surface mesh whose vertices we want to sample
+	 *  @param CellSizeAsPercentOfBounds	Grid cell size expressed as a percentage of the mesh bounding box size
+	 *  @param OutSamples					Indices of chosen vertices
+	 */
+	static void GridSample(const FDynamicMesh3& Mesh, 
+						   int GridResolutionMaxAxis,
+						   TArray<int32>& OutSamples);
+
+	/** Used for testing/debugging */
+	static FVector3i DebugGetCellIndex(const FDynamicMesh3& Mesh,
+									   int GridResolutionMaxAxis,
+									   int VertexIndex);
+
 
 public:
 	FMeshConvexHull(const FDynamicMesh3* MeshIn)
@@ -45,4 +63,4 @@ public:
 protected:
 	bool Compute_FullMesh();
 	bool Compute_VertexSubset();
-};
\ No newline at end of file
+};
diff --git a/Engine/Plugins/Experimental/GeometryProcessing/Source/GeometricObjects/Public/Util/GridIndexing3.h b/Engine/Plugins/Experimental/GeometryProcessing/Source/GeometricObjects/Public/Util/GridIndexing3.h
--- a/Engine/Plugins/Experimental/GeometryProcessing/Source/GeometricObjects/Public/Util/GridIndexing3.h
+++ b/Engine/Plugins/Experimental/GeometryProcessing/Source/GeometricObjects/Public/Util/GridIndexing3.h
@@ -5,7 +5,7 @@
 #pragma once
 
 #include "VectorTypes.h"
-
+#include "BoxTypes.h"
 
 /**
  * Convert between integer grid coordinates and scaled real-valued coordinates (ie assumes integer grid origin == real origin)
@@ -99,4 +99,35 @@ typedef TShiftGridIndexer3<float> FShiftGridIndexer3f;
 typedef TShiftGridIndexer3<double> FShiftGridIndexer3d;
 
 
+template<typename RealType>
+struct TBoundsGridIndexer3 : public TShiftGridIndexer3<RealType>
+{
+	using TShiftGridIndexer3<RealType>::CellSize;
+	using TShiftGridIndexer3<RealType>::Origin;
+
+	FVector3<RealType> BoundsMax;
+
+	TBoundsGridIndexer3(const TAxisAlignedBox3<RealType>& Bounds, RealType CellSize)
+		: TShiftGridIndexer3<RealType>(Bounds.Min, CellSize),
+		BoundsMax(Bounds.Max)
+	{}
+
+	FVector3i GridResolution() const
+	{
+		const RealType InvCellSize = 1.0 / CellSize;
+		const FVector3<RealType> Extents = BoundsMax - Origin;
+		return CeilInt(Extents * InvCellSize);
+	}
+
+	static FVector3i CeilInt(const FVector3<RealType>& V)
+	{
+		return FVector3i{ (int)TMathUtil<RealType>::Ceil(V[0]),
+			(int)TMathUtil<RealType>::Ceil(V[1]),
+			(int)TMathUtil<RealType>::Ceil(V[2]) };
+	}
+
+};
+
+typedef TBoundsGridIndexer3<float> FBoundsGridIndexer3f;
+typedef TBoundsGridIndexer3<double> FBoundsGridIndexer3d;
 
diff --git a/Engine/Plugins/Experimental/MeshLODToolset/Source/MeshLODToolset/Private/Graphs/GenerateStaticMeshLODProcess.cpp b/Engine/Plugins/Experimental/MeshLODToolset/Source/MeshLODToolset/Private/Graphs/GenerateStaticMeshLODProcess.cpp
--- a/Engine/Plugins/Experimental/MeshLODToolset/Source/MeshLODToolset/Private/Graphs/GenerateStaticMeshLODProcess.cpp
+++ b/Engine/Plugins/Experimental/MeshLODToolset/Source/MeshLODToolset/Private/Graphs/GenerateStaticMeshLODProcess.cpp
@@ -267,11 +267,14 @@ void UGenerateStaticMeshLODProcess::UpdateSettings(const FGenerateStaticMeshLODP
 		Generator->UpdateBakeCacheSettings(NewBakeSettings);
 	}
 
-
-	if (NewSettings.ConvexTriangleCount != CurrentSettings.ConvexTriangleCount)
+	if (NewSettings.ConvexTriangleCount != CurrentSettings.ConvexTriangleCount ||
+		NewSettings.bPrefilterVertices != CurrentSettings.bPrefilterVertices ||
+		NewSettings.PrefilterGridResolution != CurrentSettings.PrefilterGridResolution)
 	{
 		UE::GeometryFlow::FGenerateConvexHullsCollisionSettings NewGenConvexSettings = Generator->GetCurrentGenerateConvexCollisionSettings();
 		NewGenConvexSettings.SimplifyToTriangleCount = NewSettings.ConvexTriangleCount;
+		NewGenConvexSettings.bPrefilterVertices = NewSettings.bPrefilterVertices;
+		NewGenConvexSettings.PrefilterGridResolution = NewSettings.PrefilterGridResolution;
 		Generator->UpdateGenerateConvexCollisionSettings(NewGenConvexSettings);
 	}
 
diff --git a/Engine/Plugins/Experimental/MeshLODToolset/Source/MeshLODToolset/Private/Tools/GenerateStaticMeshLODAssetTool.cpp b/Engine/Plugins/Experimental/MeshLODToolset/Source/MeshLODToolset/Private/Tools/GenerateStaticMeshLODAssetTool.cpp
--- a/Engine/Plugins/Experimental/MeshLODToolset/Source/MeshLODToolset/Private/Tools/GenerateStaticMeshLODAssetTool.cpp
+++ b/Engine/Plugins/Experimental/MeshLODToolset/Source/MeshLODToolset/Private/Tools/GenerateStaticMeshLODAssetTool.cpp
@@ -133,7 +133,8 @@ void UGenerateStaticMeshLODAssetTool::Setup()
 	BasicProperties->WatchProperty(BasicProperties->GeneratorSettings.BakeResolution, [this](EGenerateStaticMeshLODBakeResolution) { OnSettingsModified(); });
 	BasicProperties->WatchProperty(BasicProperties->GeneratorSettings.BakeThickness, [this](float) { OnSettingsModified(); });
 	BasicProperties->WatchProperty(BasicProperties->GeneratorSettings.ConvexTriangleCount, [this](int) { OnSettingsModified(); });
-
+	BasicProperties->WatchProperty(BasicProperties->GeneratorSettings.bPrefilterVertices, [this](bool) { OnSettingsModified(); });
+	BasicProperties->WatchProperty(BasicProperties->GeneratorSettings.PrefilterGridResolution, [this](float) { OnSettingsModified(); });
 
 
 	FBoxSphereBounds Bounds = StaticMeshComponent->Bounds;
diff --git a/Engine/Plugins/Experimental/MeshLODToolset/Source/MeshLODToolset/Public/Graphs/GenerateStaticMeshLODProcess.h b/Engine/Plugins/Experimental/MeshLODToolset/Source/MeshLODToolset/Public/Graphs/GenerateStaticMeshLODProcess.h
--- a/Engine/Plugins/Experimental/MeshLODToolset/Source/MeshLODToolset/Public/Graphs/GenerateStaticMeshLODProcess.h
+++ b/Engine/Plugins/Experimental/MeshLODToolset/Source/MeshLODToolset/Public/Graphs/GenerateStaticMeshLODProcess.h
@@ -82,6 +82,14 @@ struct FGenerateStaticMeshLODProcessSettings
 
 	UPROPERTY(EditAnywhere, Category = ConvexCollision, meta = (DisplayName = "Convex Tri Count"))
 	int ConvexTriangleCount = 50;
+
+	UPROPERTY(EditAnywhere, Category = ConvexCollision)
+	bool bPrefilterVertices = false;
+
+	/** Grid resolution (along the maximum-length axis) */
+	UPROPERTY(EditAnywhere, Category = ConvexCollision, meta = (EditCondition = "bPrefilterVertices", UIMin = 1, UIMax = 30))
+	int PrefilterGridResolution = 10;
+
 };
 
 
