Geometry Flow: Add a node to get index sets from connected mesh components.

#rnx

diff --git a/Engine/Plugins/Experimental/GeometryFlow/Source/GeometryFlowMeshProcessing/Public/MeshDecompositionNodes/MakeTriangleSetsNode.h b/Engine/Plugins/Experimental/GeometryFlow/Source/GeometryFlowMeshProcessing/Public/MeshDecompositionNodes/MakeTriangleSetsNode.h
--- a/Engine/Plugins/Experimental/GeometryFlow/Source/GeometryFlowMeshProcessing/Public/MeshDecompositionNodes/MakeTriangleSetsNode.h
+++ b/Engine/Plugins/Experimental/GeometryFlow/Source/GeometryFlowMeshProcessing/Public/MeshDecompositionNodes/MakeTriangleSetsNode.h
@@ -5,7 +5,7 @@
 #include "GeometryFlowNodeUtil.h"
 #include "DataTypes/DynamicMeshData.h"
 #include "DataTypes/IndexSetsData.h"
-
+#include "Selections/MeshConnectedComponents.h"
 
 namespace UE
 {
@@ -150,7 +150,25 @@ public:
 };
 
 
+class FMakeTriangleSetsFromConnectedComponentsNode : public FMakeTriangleSetsFromMeshNode
+{
+
+public:
+
+	virtual void ComputeIndexSets(const FNamedDataMap& DatasIn, const FDynamicMesh3& Mesh, FIndexSets& SetsOut) override
+	{
+		FMeshConnectedComponents MeshRegions(&Mesh);
+		MeshRegions.FindConnectedTriangles();
+
+		for (const FMeshConnectedComponents::FComponent& Component : MeshRegions.Components)
+		{
+			SetsOut.AppendSet(Component.Indices);
+		}
+	}
+
+};
+
 
 
 }	// end namespace GeometryFlow
-}	// end namespace UE
\ No newline at end of file
+}	// end namespace UE
