AutoLOD: if no specific sets of triangles are provided to the convex hull node, just compute the convex hull of the entire mesh.

#rnx

diff --git a/Engine/Plugins/Experimental/GeometryFlow/Source/GeometryFlowMeshProcessing/Private/PhysicsNodes/GenerateConvexHullsCollisionNode.cpp b/Engine/Plugins/Experimental/GeometryFlow/Source/GeometryFlowMeshProcessing/Private/PhysicsNodes/GenerateConvexHullsCollisionNode.cpp
--- a/Engine/Plugins/Experimental/GeometryFlow/Source/GeometryFlowMeshProcessing/Private/PhysicsNodes/GenerateConvexHullsCollisionNode.cpp
+++ b/Engine/Plugins/Experimental/GeometryFlow/Source/GeometryFlowMeshProcessing/Private/PhysicsNodes/GenerateConvexHullsCollisionNode.cpp
@@ -58,23 +58,46 @@ void FGenerateConvexHullsCollisionNode::Evaluate(
 					int32 NumShapes = IndexData.IndexSets.Num();
 					TArray<FConvexShape3d> Convexes;
 					TArray<bool> ConvexOk;
-					Convexes.SetNum(NumShapes);
-					ConvexOk.SetNum(NumShapes);
-					ParallelFor(NumShapes, [&](int32 k)
+
+					if (NumShapes == 0)
 					{
-						ConvexOk[k] = false;
+						// No index sets given. Create a single convex hull from the whole mesh.
+						Convexes.SetNum(1);
+						ConvexOk.SetNum(1);
+						ConvexOk[0] = false;
 						FMeshConvexHull Hull(&Mesh);
-						MeshIndexUtil::TriangleToVertexIDs(&Mesh, IndexData.IndexSets[k], Hull.VertexSet);
 						Hull.bPostSimplify = Settings.SimplifyToTriangleCount > 0;
 						Hull.MaxTargetFaceCount = Settings.SimplifyToTriangleCount;
 						if (Hull.Compute())
 						{
 							FConvexShape3d NewConvex;
 							NewConvex.Mesh = MoveTemp(Hull.ConvexHull);
-							Convexes[k] = MoveTemp(NewConvex);
-							ConvexOk[k] = true;
+							Convexes[0] = MoveTemp(NewConvex);
+							ConvexOk[0] = true;
 						}
-					});
+					}
+					else
+					{
+						// Create one hull for each index set.
+						Convexes.SetNum(NumShapes);
+						ConvexOk.SetNum(NumShapes);
+						ParallelFor(NumShapes, [&](int32 k)
+						{
+							ConvexOk[k] = false;
+							FMeshConvexHull Hull(&Mesh);
+							MeshIndexUtil::TriangleToVertexIDs(&Mesh, IndexData.IndexSets[k], Hull.VertexSet);
+							Hull.bPostSimplify = Settings.SimplifyToTriangleCount > 0;
+							Hull.MaxTargetFaceCount = Settings.SimplifyToTriangleCount;
+							if (Hull.Compute())
+							{
+								FConvexShape3d NewConvex;
+								NewConvex.Mesh = MoveTemp(Hull.ConvexHull);
+								Convexes[k] = MoveTemp(NewConvex);
+								ConvexOk[k] = true;
+							}
+						});
+					}
+
 					Result.Geometry.Convexes = MoveTemp(Convexes);
 				}
 				SetOutput(OutParamGeometry(), MakeMovableData<FCollisionGeometry>(MoveTemp(Result)));
