GeometryFlow: add vertex weight map data type, and a node for getting a set of triangles from a mesh and weight map.

#rnx

diff --git a/Engine/Plugins/Experimental/GeometryFlow/Source/GeometryFlowMeshProcessing/Public/DataTypes/WeightMapData.h b/Engine/Plugins/Experimental/GeometryFlow/Source/GeometryFlowMeshProcessing/Public/DataTypes/WeightMapData.h
new file mode 100644
--- /dev/null
+++ b/Engine/Plugins/Experimental/GeometryFlow/Source/GeometryFlowMeshProcessing/Public/DataTypes/WeightMapData.h
@@ -0,0 +1,29 @@
+// Copyright Epic Games, Inc. All Rights Reserved.
+
+#pragma once
+
+#include "GeometryFlowCoreNodes.h"
+#include "MeshProcessingNodes/MeshProcessingDataTypes.h"
+
+
+namespace UE
+{
+namespace GeometryFlow
+{
+
+
+struct FWeightMap
+{
+	DECLARE_GEOMETRYFLOW_DATA_TYPE_IDENTIFIER(EMeshProcessingDataTypes::WeightMap);
+
+	TArray<float> Weights;
+
+	// TODO: Sparse map from vertex ID to weight value?
+};
+
+
+GEOMETRYFLOW_DECLARE_BASIC_TYPES(WeightMap, FWeightMap, (int)EMeshProcessingDataTypes::WeightMap)
+
+
+}	// end namespace GeometryFlow
+}	// end namespace UE
diff --git a/Engine/Plugins/Experimental/GeometryFlow/Source/GeometryFlowMeshProcessing/Public/MeshDecompositionNodes/MakeTriangleSetsNode.h b/Engine/Plugins/Experimental/GeometryFlow/Source/GeometryFlowMeshProcessing/Public/MeshDecompositionNodes/MakeTriangleSetsNode.h
--- a/Engine/Plugins/Experimental/GeometryFlow/Source/GeometryFlowMeshProcessing/Public/MeshDecompositionNodes/MakeTriangleSetsNode.h
+++ b/Engine/Plugins/Experimental/GeometryFlow/Source/GeometryFlowMeshProcessing/Public/MeshDecompositionNodes/MakeTriangleSetsNode.h
@@ -5,6 +5,7 @@
 #include "GeometryFlowNodeUtil.h"
 #include "DataTypes/DynamicMeshData.h"
 #include "DataTypes/IndexSetsData.h"
+#include "DataTypes/WeightMapData.h"
 #include "Selections/MeshConnectedComponents.h"
 
 namespace UE
@@ -169,6 +170,59 @@ public:
 };
 
 
+/// If one triangle vertex has a weight greater than the given threshold, include it in the output triangle set.
+/// TODO: Optionally make it so the *average* triangle vertex weight has to exceed the threshold. Or the minumum triangle
+/// vertex weight.
+class FMakeTriangleSetsFromWeightMapNode : public FMakeTriangleSetsFromMeshNode
+{
+
+public:
+
+	static const FString InParamWeightMap() { return TEXT("WeightMap"); }
+	static const FString InParamThreshold() { return TEXT("Threshold"); }
+
+	FMakeTriangleSetsFromWeightMapNode() : FMakeTriangleSetsFromMeshNode()
+	{
+		AddInput(InParamWeightMap(), MakeBasicInput<FWeightMap>());
+		AddInput(InParamThreshold(), MakeUnique<TBasicNodeInput<float, (int)EDataTypes::Float>>());
+	}
+
+	virtual void CheckAdditionalInputs(const FNamedDataMap& DatasIn, bool& bRecomputeRequired, bool& bAllInputsValid) override
+	{
+		FindAndUpdateInputForEvaluate(InParamWeightMap(), DatasIn, bRecomputeRequired, bAllInputsValid);
+		FindAndUpdateInputForEvaluate(InParamThreshold(), DatasIn, bRecomputeRequired, bAllInputsValid);
+	}
+
+	virtual void ComputeIndexSets(const FNamedDataMap& DatasIn, const FDynamicMesh3& Mesh, FIndexSets& SetsOut) override
+	{
+		TSafeSharedPtr<IData> WeightMapArg = DatasIn.FindData(InParamWeightMap());
+		const FWeightMap& WeightMap = WeightMapArg->GetDataConstRef<FWeightMap>(FWeightMap::DataTypeIdentifier);
+		const TArray<float>& Weights = WeightMap.Weights;
+		check(Weights.Num() >= Mesh.MaxVertexID());
+
+		TSafeSharedPtr<IData> ThresholdArg = DatasIn.FindData(InParamThreshold());
+		const float Threshold = ThresholdArg->GetDataConstRef<float>((int)EDataTypes::Float);
+
+		TArray<int32> Set;
+
+		for (int TriangleID : Mesh.TriangleIndicesItr())
+		{
+			FIndex3i Triangle = Mesh.GetTriangle(TriangleID);
+			for (int VertexID : {Triangle[0], Triangle[1], Triangle[2]} )
+			{
+				if (Weights[VertexID] > Threshold)
+				{
+					Set.Add(TriangleID);
+					break;
+				}
+			}
+		}
+
+		SetsOut.AppendSet(Set);
+	}
+
+};
+
 
 }	// end namespace GeometryFlow
 }	// end namespace UE
diff --git a/Engine/Plugins/Experimental/GeometryFlow/Source/GeometryFlowMeshProcessing/Public/MeshProcessingNodes/MeshProcessingDataTypes.h b/Engine/Plugins/Experimental/GeometryFlow/Source/GeometryFlowMeshProcessing/Public/MeshProcessingNodes/MeshProcessingDataTypes.h
--- a/Engine/Plugins/Experimental/GeometryFlow/Source/GeometryFlowMeshProcessing/Public/MeshProcessingNodes/MeshProcessingDataTypes.h
+++ b/Engine/Plugins/Experimental/GeometryFlow/Source/GeometryFlowMeshProcessing/Public/MeshProcessingNodes/MeshProcessingDataTypes.h
@@ -21,7 +21,7 @@ enum class EMeshProcessingDataTypes
 	NormalMapImage = (int)EDataTypes::BaseMeshProcessingTypes + 12,
 	CollisionGeometry = (int)EDataTypes::BaseMeshProcessingTypes + 13,
 	IndexSets = (int)EDataTypes::BaseMeshProcessingTypes + 14,
-
+	WeightMap = (int)EDataTypes::BaseMeshProcessingTypes + 15,
 
 	SolidifySettings = (int)EDataTypes::BaseMeshProcessingTypes + 100,
 	VoxMorphologyOpSettings = (int)EDataTypes::BaseMeshProcessingTypes + 101,
