Initial implementation of simple collision geometry generator node. Just wraps FMeshSimpleShapeApproximation. The only option right now is which type of geometric primitive will be produced. Does not support ConvexHull, as there is a separate node for that.

#rnx

diff --git a/Engine/Plugins/Experimental/GeometryFlow/Source/GeometryFlowMeshProcessing/Private/PhysicsNodes/GenerateSimpleCollisionNode.cpp b/Engine/Plugins/Experimental/GeometryFlow/Source/GeometryFlowMeshProcessing/Private/PhysicsNodes/GenerateSimpleCollisionNode.cpp
new file mode 100644
--- /dev/null
+++ b/Engine/Plugins/Experimental/GeometryFlow/Source/GeometryFlowMeshProcessing/Private/PhysicsNodes/GenerateSimpleCollisionNode.cpp
@@ -0,0 +1,69 @@
+// Copyright Epic Games, Inc. All Rights Reserved.
+
+#include "PhysicsNodes/GenerateSimpleCollisionNode.h"
+#include "ShapeApproximation/MeshSimpleShapeApproximation.h"
+
+using namespace UE::GeometryFlow;
+
+void FGenerateSimpleCollisionNode::Evaluate(
+	const FNamedDataMap& DatasIn,
+	FNamedDataMap& DatasOut,
+	TUniquePtr<FEvaluationInfo>& EvaluationInfo)
+{
+	if (ensure(DatasOut.Contains(OutParamGeometry())))
+	{
+		bool bAllInputsValid = true;
+		bool bRecomputeRequired = (IsOutputAvailable(OutParamGeometry()) == false);
+
+		TSafeSharedPtr<IData> MeshArg = FindAndUpdateInputForEvaluate(InParamMesh(), 
+																	  DatasIn, 
+																	  bRecomputeRequired, 
+																	  bAllInputsValid);
+		const FDynamicMesh3& Mesh = MeshArg->GetDataConstRef<FDynamicMesh3>((int)EMeshProcessingDataTypes::DynamicMesh);
+
+		TSafeSharedPtr<IData> SettingsArg = FindAndUpdateInputForEvaluate(InParamSettings(), 
+																		  DatasIn, 
+																		  bRecomputeRequired, bAllInputsValid);
+		FGenerateSimpleCollisionSettings Settings;
+		SettingsArg->GetDataCopy(Settings, FGenerateSimpleCollisionSettings::DataTypeIdentifier);
+	
+		if (bAllInputsValid)
+		{
+			if (bRecomputeRequired)
+			{
+				FMeshSimpleShapeApproximation ShapeApproximator;
+				ShapeApproximator.bDetectSpheres = false;
+				ShapeApproximator.bDetectBoxes = false;
+				ShapeApproximator.bDetectCapsules = false;
+				ShapeApproximator.bDetectConvexes = false;
+
+				ShapeApproximator.InitializeSourceMeshes({ &Mesh });
+
+				FCollisionGeometry CollisionGeometry;
+				switch (Settings.Type)
+				{
+				case ESimpleCollisionGeometryType::AlignedBoxes:
+					ShapeApproximator.bDetectBoxes = true;
+					ShapeApproximator.Generate_AlignedBoxes(CollisionGeometry.Geometry);
+					break;
+				case ESimpleCollisionGeometryType::OrientedBoxes:
+					ShapeApproximator.bDetectBoxes = true;
+					ShapeApproximator.Generate_OrientedBoxes(CollisionGeometry.Geometry);
+					break;
+				case ESimpleCollisionGeometryType::MinimalSpheres:
+					ShapeApproximator.bDetectSpheres = true;
+					ShapeApproximator.Generate_MinimalSpheres(CollisionGeometry.Geometry);
+					break;
+				case ESimpleCollisionGeometryType::Capsules:
+					ShapeApproximator.bDetectCapsules = true;
+					ShapeApproximator.Generate_Capsules(CollisionGeometry.Geometry);
+					break;
+				}
+
+				SetOutput(OutParamGeometry(), MakeMovableData<FCollisionGeometry>(MoveTemp(CollisionGeometry)));
+			}
+
+			DatasOut.SetData(OutParamGeometry(), GetOutput(OutParamGeometry()));
+		}
+	}
+}
diff --git a/Engine/Plugins/Experimental/GeometryFlow/Source/GeometryFlowMeshProcessing/Public/MeshProcessingNodes/MeshProcessingDataTypes.h b/Engine/Plugins/Experimental/GeometryFlow/Source/GeometryFlowMeshProcessing/Public/MeshProcessingNodes/MeshProcessingDataTypes.h
--- a/Engine/Plugins/Experimental/GeometryFlow/Source/GeometryFlowMeshProcessing/Public/MeshProcessingNodes/MeshProcessingDataTypes.h
+++ b/Engine/Plugins/Experimental/GeometryFlow/Source/GeometryFlowMeshProcessing/Public/MeshProcessingNodes/MeshProcessingDataTypes.h
@@ -45,9 +45,10 @@ enum class EMeshProcessingDataTypes
 
 
 	GenerateCollisionConvexHullsSettings = (int)EDataTypes::BaseMeshProcessingTypes + 200,
+	GenerateSimpleCollisionSettings = (int)EDataTypes::BaseMeshProcessingTypes + 201,
 };
 
 
 
 }	// end namespace GeometryFlow
-}	// en
\ No newline at end of file
+}	// end namespace UE
diff --git a/Engine/Plugins/Experimental/GeometryFlow/Source/GeometryFlowMeshProcessing/Public/PhysicsNodes/GenerateSimpleCollisionNode.h b/Engine/Plugins/Experimental/GeometryFlow/Source/GeometryFlowMeshProcessing/Public/PhysicsNodes/GenerateSimpleCollisionNode.h
new file mode 100644
--- /dev/null
+++ b/Engine/Plugins/Experimental/GeometryFlow/Source/GeometryFlowMeshProcessing/Public/PhysicsNodes/GenerateSimpleCollisionNode.h
@@ -0,0 +1,65 @@
+// Copyright Epic Games, Inc. All Rights Reserved.
+
+#pragma once
+
+#include "GeometryFlowCoreNodes.h"
+#include "GeometryFlowNodeUtil.h"
+#include "MeshProcessingNodes/MeshProcessingDataTypes.h"
+#include "DataTypes/DynamicMeshData.h"
+#include "DataTypes/CollisionGeometryData.h"
+#include "DataTypes/IndexSetsData.h"
+
+
+namespace UE
+{
+namespace GeometryFlow
+{
+
+enum class ESimpleCollisionGeometryType
+{
+	AlignedBoxes = 1,
+	OrientedBoxes = 2,
+	MinimalSpheres = 3,
+	Capsules = 4
+};
+
+
+struct GEOMETRYFLOWMESHPROCESSING_API FGenerateSimpleCollisionSettings
+{
+	DECLARE_GEOMETRYFLOW_DATA_TYPE_IDENTIFIER(EMeshProcessingDataTypes::GenerateSimpleCollisionSettings);
+
+	ESimpleCollisionGeometryType Type = ESimpleCollisionGeometryType::MinimalSpheres;
+};
+GEOMETRYFLOW_DECLARE_SETTINGS_TYPES(FGenerateSimpleCollisionSettings, GenerateSimpleCollision);
+
+
+class GEOMETRYFLOWMESHPROCESSING_API FGenerateSimpleCollisionNode : public FNode
+{
+protected:
+	using SettingsDataType = TMovableData<FGenerateSimpleCollisionSettings, FGenerateSimpleCollisionSettings::DataTypeIdentifier>;
+
+public:
+	static const FString InParamMesh() { return TEXT("Mesh"); }
+	// TODO: static const FString InParamIndexSets() { return TEXT("TriangleSets"); }	
+	static const FString InParamSettings() { return TEXT("Settings"); }
+	static const FString OutParamGeometry() { return TEXT("Geometry"); }
+
+public:
+
+	FGenerateSimpleCollisionNode()
+	{
+		AddInput(InParamMesh(), MakeUnique<FDynamicMeshInput>());
+		// TODO: AddInput(InParamIndexSets(), MakeBasicInput<FIndexSets>());
+		AddInput(InParamSettings(), MakeBasicInput<FGenerateSimpleCollisionSettings>());
+
+		AddOutput(OutParamGeometry(), MakeBasicOutput<FCollisionGeometry>());
+	}
+
+	virtual void Evaluate(
+		const FNamedDataMap& DatasIn,
+		FNamedDataMap& DatasOut,
+		TUniquePtr<FEvaluationInfo>& EvaluationInfo) override;
+};
+
+}	// end namespace GeometryFlow
+}	// end namespace UE
