Add support for convex hulls and triangle sets to FGenerateSimpleCollisionNode. Make this the default collision geometry generation node to replace the GenerateConvexCollisionNode.

#rnx

diff --git a/Engine/Plugins/Experimental/GeometryFlow/Source/GeometryFlowMeshProcessing/Private/PhysicsNodes/GenerateSimpleCollisionNode.cpp b/Engine/Plugins/Experimental/GeometryFlow/Source/GeometryFlowMeshProcessing/Private/PhysicsNodes/GenerateSimpleCollisionNode.cpp
--- a/Engine/Plugins/Experimental/GeometryFlow/Source/GeometryFlowMeshProcessing/Private/PhysicsNodes/GenerateSimpleCollisionNode.cpp
+++ b/Engine/Plugins/Experimental/GeometryFlow/Source/GeometryFlowMeshProcessing/Private/PhysicsNodes/GenerateSimpleCollisionNode.cpp
@@ -2,9 +2,112 @@
 
 #include "PhysicsNodes/GenerateSimpleCollisionNode.h"
 #include "ShapeApproximation/MeshSimpleShapeApproximation.h"
+#include "Operations/MeshConvexHull.h"
+#include "MeshIndexUtil.h"
+#include "Async/ParallelFor.h"
+#include "DynamicSubmesh3.h"
 
 using namespace UE::GeometryFlow;
 
+
+namespace GenerateSimpleCollisionNodeHelpers
+{
+	TArray<int> SetIntersection(const TArray<int>& A, const TArray<int>& B)
+	{
+		TSet<int> SA(A);
+		TSet<int> SB(B);
+		return SA.Intersect(SB).Array();
+	}
+
+	void GenerateConvexHulls(const FDynamicMesh3& Mesh,
+							 const FIndexSets& IndexData,
+							 const FGenerateSimpleCollisionSettings& Settings,
+							 FCollisionGeometry& Result)
+	{
+		int32 NumShapes = IndexData.IndexSets.Num();
+		TArray<FConvexShape3d> Convexes;
+		TArray<bool> ConvexOk;
+
+		if (NumShapes == 0)
+		{
+			// No index sets given. Create a single convex hull from the whole mesh.
+			Convexes.SetNum(1);
+			ConvexOk.SetNum(1);
+			ConvexOk[0] = false;
+			FMeshConvexHull Hull(&Mesh);
+
+			if (Settings.ConvexHullSettings.bPrefilterVertices)
+			{
+				FMeshConvexHull::GridSample(Mesh, Settings.ConvexHullSettings.PrefilterGridResolution, Hull.VertexSet);
+			}
+
+			Hull.bPostSimplify = Settings.ConvexHullSettings.SimplifyToTriangleCount > 0;
+			Hull.MaxTargetFaceCount = Settings.ConvexHullSettings.SimplifyToTriangleCount;
+			if (Hull.Compute())
+			{
+				FConvexShape3d NewConvex;
+				NewConvex.Mesh = MoveTemp(Hull.ConvexHull);
+				Convexes[0] = MoveTemp(NewConvex);
+				ConvexOk[0] = true;
+			}
+		}
+		else
+		{
+			// Create one hull for each index set.
+			Convexes.SetNum(NumShapes);
+			ConvexOk.SetNum(NumShapes);
+			ParallelFor(NumShapes, [&](int32 k)
+			{
+				ConvexOk[k] = false;
+				FMeshConvexHull Hull(&Mesh);
+
+				if (Settings.ConvexHullSettings.bPrefilterVertices)
+				{
+					TArray<int> SampledVertices;
+					FMeshConvexHull::GridSample(Mesh, Settings.ConvexHullSettings.PrefilterGridResolution, SampledVertices);
+					TArray<int> IndexSetVertices;
+					MeshIndexUtil::TriangleToVertexIDs(&Mesh, IndexData.IndexSets[k], IndexSetVertices);
+					Hull.VertexSet = SetIntersection(SampledVertices, IndexSetVertices);
+				}
+				else
+				{
+					MeshIndexUtil::TriangleToVertexIDs(&Mesh, IndexData.IndexSets[k], Hull.VertexSet);
+				}
+
+				Hull.bPostSimplify = Settings.ConvexHullSettings.SimplifyToTriangleCount > 0;
+				Hull.MaxTargetFaceCount = Settings.ConvexHullSettings.SimplifyToTriangleCount;
+				if (Hull.Compute())
+				{
+					FConvexShape3d NewConvex;
+					NewConvex.Mesh = MoveTemp(Hull.ConvexHull);
+					Convexes[k] = MoveTemp(NewConvex);
+					ConvexOk[k] = true;
+				}
+			});
+		}
+
+		Result.Geometry.Convexes = MoveTemp(Convexes);
+	}
+
+	void GenerateSubMeshes(const FDynamicMesh3& Mesh, const FIndexSets& TriangleIndexSets, TArray<FDynamicMesh3>& OutSubMeshes)
+	{
+		if (TriangleIndexSets.IndexSets.Num() == 0)
+		{
+			OutSubMeshes.Add(Mesh);
+		}
+		else
+		{
+			for (const TArray<int32>& Set : TriangleIndexSets.IndexSets)
+			{
+				FDynamicSubmesh3 Sub(&Mesh, Set);
+				OutSubMeshes.Add(Sub.GetSubmesh());
+			}
+		}
+	}
+
+} // namespace GenerateSimpleCollisionNodeHelpers
+
+
 void FGenerateSimpleCollisionNode::Evaluate(
 	const FNamedDataMap& DatasIn,
 	FNamedDataMap& DatasOut,
@@ -15,29 +118,45 @@ void FGenerateSimpleCollisionNode::Evaluate(
 		bool bAllInputsValid = true;
 		bool bRecomputeRequired = (IsOutputAvailable(OutParamGeometry()) == false);
 
-		TSafeSharedPtr<IData> MeshArg = FindAndUpdateInputForEvaluate(InParamMesh(), 
-																	  DatasIn, 
-																	  bRecomputeRequired, 
+		TSafeSharedPtr<IData> MeshArg = FindAndUpdateInputForEvaluate(InParamMesh(),
+																	  DatasIn,
+																	  bRecomputeRequired,
 																	  bAllInputsValid);
 		const FDynamicMesh3& Mesh = MeshArg->GetDataConstRef<FDynamicMesh3>((int)EMeshProcessingDataTypes::DynamicMesh);
 
-		TSafeSharedPtr<IData> SettingsArg = FindAndUpdateInputForEvaluate(InParamSettings(), 
-																		  DatasIn, 
+		TSafeSharedPtr<IData> TriSetsArg = FindAndUpdateInputForEvaluate(InParamIndexSets(), 
+																		 DatasIn, 
+																		 bRecomputeRequired, 
+																		 bAllInputsValid);
+		const FIndexSets& IndexData = TriSetsArg->GetDataConstRef<FIndexSets>((int)EMeshProcessingDataTypes::IndexSets);
+
+		TSafeSharedPtr<IData> SettingsArg = FindAndUpdateInputForEvaluate(InParamSettings(),
+																		  DatasIn,
 																		  bRecomputeRequired, bAllInputsValid);
 		FGenerateSimpleCollisionSettings Settings;
 		SettingsArg->GetDataCopy(Settings, FGenerateSimpleCollisionSettings::DataTypeIdentifier);
-	
+
 		if (bAllInputsValid)
 		{
 			if (bRecomputeRequired)
 			{
 				FMeshSimpleShapeApproximation ShapeApproximator;
+
+				// Don't cache anything, just compute when asked to
 				ShapeApproximator.bDetectSpheres = false;
 				ShapeApproximator.bDetectBoxes = false;
 				ShapeApproximator.bDetectCapsules = false;
 				ShapeApproximator.bDetectConvexes = false;
 
-				ShapeApproximator.InitializeSourceMeshes({ &Mesh });
+				TArray<FDynamicMesh3> SubMeshes;
+				GenerateSimpleCollisionNodeHelpers::GenerateSubMeshes(Mesh, IndexData, SubMeshes);
+
+				TArray<const FDynamicMesh3*> SubMeshPointers;
+				for (FDynamicMesh3& SubMesh : SubMeshes)
+				{
+					SubMeshPointers.Add(&SubMesh);
+				}
+				ShapeApproximator.InitializeSourceMeshes(SubMeshPointers);
 
 				FCollisionGeometry CollisionGeometry;
 				switch (Settings.Type)
@@ -58,6 +177,9 @@ void FGenerateSimpleCollisionNode::Evaluate(
 					ShapeApproximator.bDetectCapsules = true;
 					ShapeApproximator.Generate_Capsules(CollisionGeometry.Geometry);
 					break;
+				case ESimpleCollisionGeometryType::ConvexHulls:
+					GenerateSimpleCollisionNodeHelpers::GenerateConvexHulls(Mesh, IndexData, Settings, CollisionGeometry);
+					break;
 				}
 
 				SetOutput(OutParamGeometry(), MakeMovableData<FCollisionGeometry>(MoveTemp(CollisionGeometry)));
diff --git a/Engine/Plugins/Experimental/GeometryFlow/Source/GeometryFlowMeshProcessing/Public/PhysicsNodes/GenerateSimpleCollisionNode.h b/Engine/Plugins/Experimental/GeometryFlow/Source/GeometryFlowMeshProcessing/Public/PhysicsNodes/GenerateSimpleCollisionNode.h
--- a/Engine/Plugins/Experimental/GeometryFlow/Source/GeometryFlowMeshProcessing/Public/PhysicsNodes/GenerateSimpleCollisionNode.h
+++ b/Engine/Plugins/Experimental/GeometryFlow/Source/GeometryFlowMeshProcessing/Public/PhysicsNodes/GenerateSimpleCollisionNode.h
@@ -15,12 +15,13 @@ namespace UE
 namespace GeometryFlow
 {
 
-enum class ESimpleCollisionGeometryType
+enum class ESimpleCollisionGeometryType : uint8
 {
-	AlignedBoxes = 1,
-	OrientedBoxes = 2,
-	MinimalSpheres = 3,
-	Capsules = 4
+	AlignedBoxes,
+	OrientedBoxes,
+	MinimalSpheres,
+	Capsules,
+	ConvexHulls
 };
 
 
@@ -28,8 +29,17 @@ struct GEOMETRYFLOWMESHPROCESSING_API FGenerateSimpleCollisionSettings
 {
 	DECLARE_GEOMETRYFLOW_DATA_TYPE_IDENTIFIER(EMeshProcessingDataTypes::GenerateSimpleCollisionSettings);
 
-	ESimpleCollisionGeometryType Type = ESimpleCollisionGeometryType::MinimalSpheres;
+	ESimpleCollisionGeometryType Type = ESimpleCollisionGeometryType::ConvexHulls;
+
+	struct FGenerateConvexHullSettings
+	{
+		int32 SimplifyToTriangleCount = 50;
+		bool bPrefilterVertices = false;
+		int PrefilterGridResolution = 10;
+	} ConvexHullSettings;
+
 };
+
 GEOMETRYFLOW_DECLARE_SETTINGS_TYPES(FGenerateSimpleCollisionSettings, GenerateSimpleCollision);
 
 
@@ -40,7 +50,7 @@ protected:
 
 public:
 	static const FString InParamMesh() { return TEXT("Mesh"); }
-	// TODO: static const FString InParamIndexSets() { return TEXT("TriangleSets"); }	
+	static const FString InParamIndexSets() { return TEXT("TriangleSets"); }	
 	static const FString InParamSettings() { return TEXT("Settings"); }
 	static const FString OutParamGeometry() { return TEXT("Geometry"); }
 
@@ -49,7 +59,7 @@ public:
 	FGenerateSimpleCollisionNode()
 	{
 		AddInput(InParamMesh(), MakeUnique<FDynamicMeshInput>());
-		// TODO: AddInput(InParamIndexSets(), MakeBasicInput<FIndexSets>());
+		AddInput(InParamIndexSets(), MakeBasicInput<FIndexSets>());
 		AddInput(InParamSettings(), MakeBasicInput<FGenerateSimpleCollisionSettings>());
 
 		AddOutput(OutParamGeometry(), MakeBasicOutput<FCollisionGeometry>());
diff --git a/Engine/Plugins/Experimental/MeshLODToolset/Source/MeshLODToolset/Private/Graphs/GenerateMeshLODGraph.cpp b/Engine/Plugins/Experimental/MeshLODToolset/Source/MeshLODToolset/Private/Graphs/GenerateMeshLODGraph.cpp
--- a/Engine/Plugins/Experimental/MeshLODToolset/Source/MeshLODToolset/Private/Graphs/GenerateMeshLODGraph.cpp
+++ b/Engine/Plugins/Experimental/MeshLODToolset/Source/MeshLODToolset/Private/Graphs/GenerateMeshLODGraph.cpp
@@ -24,7 +24,7 @@
 #include "MeshBakingNodes/BakeMeshTextureImageNode.h"
 
 #include "MeshDecompositionNodes/MakeTriangleSetsNode.h"
-#include "PhysicsNodes/GenerateConvexHullsCollisionNode.h"
+#include "PhysicsNodes/GenerateSimpleCollisionNode.h"
 #include "GeometryFlowExecutor.h"
 
 using namespace UE::GeometryFlow;
@@ -226,10 +226,10 @@ void FGenerateMeshLODGraph::UpdateBakeCacheSettings(const UE::GeometryFlow::FMes
 }
 
 
-void FGenerateMeshLODGraph::UpdateGenerateConvexCollisionSettings(const FGenerateConvexHullsCollisionSettings& GenConvexesSettings)
+void FGenerateMeshLODGraph::UpdateGenerateSimpleCollisionSettings(const FGenerateSimpleCollisionSettings& GenSimpleCollisionSettings)
 {
-	UpdateSettingsSourceNodeValue(*Graph, GenerateConvexesSettingsNode, GenConvexesSettings);
-	CurrentGenerateConvexHullsSettings = GenConvexesSettings;
+	UpdateSettingsSourceNodeValue(*Graph, GenerateSimpleCollisionSettingsNode, GenSimpleCollisionSettings);
+	CurrentGenerateSimpleCollisionSettings = GenSimpleCollisionSettings;
 }
 
 
@@ -408,16 +408,16 @@ void FGenerateMeshLODGraph::BuildGraph()
 	Graph->InferConnection(MeshSourceNode, DecomposeMeshForCollisionNode);
 	Graph->InferConnection(IgnoreGroupsForCollisionNode, DecomposeMeshForCollisionNode);
 
-	GenerateConvexesNode = Graph->AddNodeOfType<FGenerateConvexHullsCollisionNode>(TEXT("GenerateConvexes"));
-	Graph->InferConnection(MeshSourceNode, GenerateConvexesNode);
-	Graph->InferConnection(DecomposeMeshForCollisionNode, GenerateConvexesNode);
-	GenerateConvexesSettingsNode = Graph->AddNodeOfType<FGenerateConvexHullsCollisionSettingsSourceNode>(TEXT("GenerateConvexesSettings"));
-	Graph->InferConnection(GenerateConvexesSettingsNode, GenerateConvexesNode);
+	GenerateSimpleCollisionNode = Graph->AddNodeOfType<FGenerateSimpleCollisionNode>(TEXT("GenerateSimpleCollision"));
+	Graph->InferConnection(MeshSourceNode, GenerateSimpleCollisionNode);
+	Graph->InferConnection(DecomposeMeshForCollisionNode, GenerateSimpleCollisionNode);
+	GenerateSimpleCollisionSettingsNode = Graph->AddNodeOfType<FGenerateSimpleCollisionSettingsSourceNode>(TEXT("GenerateSimpleCollisionSettings"));
+	Graph->InferConnection(GenerateSimpleCollisionSettingsNode, GenerateSimpleCollisionNode);
 
 	// final collision output
 
 	CollisionOutputNode = Graph->AddNodeOfType<FCollisionGeometryTransferNode>(TEXT("OutputCollision"));
-	Graph->InferConnection(GenerateConvexesNode, CollisionOutputNode);
+	Graph->InferConnection(GenerateSimpleCollisionNode, CollisionOutputNode);
 
 
 	//
@@ -473,8 +473,8 @@ void FGenerateMeshLODGraph::BuildGraph()
 	IgnoreGroupsForCollision.AppendSet({ 0 });
 	UpdateSettingsSourceNodeValue(*Graph, IgnoreGroupsForCollisionNode, IgnoreGroupsForCollision);
 
-	FGenerateConvexHullsCollisionSettings GenConvexesSettings;
-	UpdateGenerateConvexCollisionSettings(GenConvexesSettings);
+	FGenerateSimpleCollisionSettings GenSimpleCollisionSettings;
+	UpdateGenerateSimpleCollisionSettings(GenSimpleCollisionSettings);
 
 
 	//FString GraphDump = Graph->DebugDumpGraph([](TSafeSharedPtr<FNode> Node)
diff --git a/Engine/Plugins/Experimental/MeshLODToolset/Source/MeshLODToolset/Private/Graphs/GenerateStaticMeshLODProcess.cpp b/Engine/Plugins/Experimental/MeshLODToolset/Source/MeshLODToolset/Private/Graphs/GenerateStaticMeshLODProcess.cpp
--- a/Engine/Plugins/Experimental/MeshLODToolset/Source/MeshLODToolset/Private/Graphs/GenerateStaticMeshLODProcess.cpp
+++ b/Engine/Plugins/Experimental/MeshLODToolset/Source/MeshLODToolset/Private/Graphs/GenerateStaticMeshLODProcess.cpp
@@ -213,7 +213,10 @@ bool UGenerateStaticMeshLODProcess::InitializeGenerator()
 	CurrentSettings.BakeResolution = (EGenerateStaticMeshLODBakeResolution)Generator->GetCurrentBakeCacheSettings().Dimensions.GetWidth();
 	CurrentSettings.BakeThickness = Generator->GetCurrentBakeCacheSettings().Thickness;
 
-	CurrentSettings.ConvexTriangleCount = Generator->GetCurrentGenerateConvexCollisionSettings().SimplifyToTriangleCount;
+	CurrentSettings.CollisionType = static_cast<EGenerateStaticMeshSimpleCollisionGeometryType>(Generator->GetCurrentGenerateSimpleCollisionSettings().Type);
+	CurrentSettings.ConvexTriangleCount = Generator->GetCurrentGenerateSimpleCollisionSettings().ConvexHullSettings.SimplifyToTriangleCount;
+	CurrentSettings.bPrefilterVertices = Generator->GetCurrentGenerateSimpleCollisionSettings().ConvexHullSettings.bPrefilterVertices;
+	CurrentSettings.PrefilterGridResolution = Generator->GetCurrentGenerateSimpleCollisionSettings().ConvexHullSettings.PrefilterGridResolution;
 
 
 	return true;
@@ -269,13 +272,15 @@ void UGenerateStaticMeshLODProcess::UpdateSettings(const FGenerateStaticMeshLODP
 
 	if (NewSettings.ConvexTriangleCount != CurrentSettings.ConvexTriangleCount ||
 		NewSettings.bPrefilterVertices != CurrentSettings.bPrefilterVertices ||
-		NewSettings.PrefilterGridResolution != CurrentSettings.PrefilterGridResolution)
+		NewSettings.PrefilterGridResolution != CurrentSettings.PrefilterGridResolution ||
+		NewSettings.CollisionType != CurrentSettings.CollisionType)
 	{
-		UE::GeometryFlow::FGenerateConvexHullsCollisionSettings NewGenConvexSettings = Generator->GetCurrentGenerateConvexCollisionSettings();
-		NewGenConvexSettings.SimplifyToTriangleCount = NewSettings.ConvexTriangleCount;
-		NewGenConvexSettings.bPrefilterVertices = NewSettings.bPrefilterVertices;
-		NewGenConvexSettings.PrefilterGridResolution = NewSettings.PrefilterGridResolution;
-		Generator->UpdateGenerateConvexCollisionSettings(NewGenConvexSettings);
+		UE::GeometryFlow::FGenerateSimpleCollisionSettings NewGenCollisionSettings = Generator->GetCurrentGenerateSimpleCollisionSettings();
+		NewGenCollisionSettings.Type = static_cast<UE::GeometryFlow::ESimpleCollisionGeometryType>(NewSettings.CollisionType);
+		NewGenCollisionSettings.ConvexHullSettings.SimplifyToTriangleCount = NewSettings.ConvexTriangleCount;
+		NewGenCollisionSettings.ConvexHullSettings.bPrefilterVertices = NewSettings.bPrefilterVertices;
+		NewGenCollisionSettings.ConvexHullSettings.PrefilterGridResolution = NewSettings.PrefilterGridResolution;
+		Generator->UpdateGenerateSimpleCollisionSettings(NewGenCollisionSettings);
 	}
 
 	CurrentSettings = NewSettings;
diff --git a/Engine/Plugins/Experimental/MeshLODToolset/Source/MeshLODToolset/Private/Tools/GenerateStaticMeshLODAssetTool.cpp b/Engine/Plugins/Experimental/MeshLODToolset/Source/MeshLODToolset/Private/Tools/GenerateStaticMeshLODAssetTool.cpp
--- a/Engine/Plugins/Experimental/MeshLODToolset/Source/MeshLODToolset/Private/Tools/GenerateStaticMeshLODAssetTool.cpp
+++ b/Engine/Plugins/Experimental/MeshLODToolset/Source/MeshLODToolset/Private/Tools/GenerateStaticMeshLODAssetTool.cpp
@@ -132,6 +132,7 @@ void UGenerateStaticMeshLODAssetTool::Setup()
 	BasicProperties->WatchProperty(BasicProperties->GeneratorSettings.NumAutoUVCharts, [this](int) { OnSettingsModified(); });
 	BasicProperties->WatchProperty(BasicProperties->GeneratorSettings.BakeResolution, [this](EGenerateStaticMeshLODBakeResolution) { OnSettingsModified(); });
 	BasicProperties->WatchProperty(BasicProperties->GeneratorSettings.BakeThickness, [this](float) { OnSettingsModified(); });
+	BasicProperties->WatchProperty(BasicProperties->GeneratorSettings.CollisionType, [this](EGenerateStaticMeshSimpleCollisionGeometryType) { OnSettingsModified(); });
 	BasicProperties->WatchProperty(BasicProperties->GeneratorSettings.ConvexTriangleCount, [this](int) { OnSettingsModified(); });
 	BasicProperties->WatchProperty(BasicProperties->GeneratorSettings.bPrefilterVertices, [this](bool) { OnSettingsModified(); });
 	BasicProperties->WatchProperty(BasicProperties->GeneratorSettings.PrefilterGridResolution, [this](float) { OnSettingsModified(); });
diff --git a/Engine/Plugins/Experimental/MeshLODToolset/Source/MeshLODToolset/Public/Graphs/GenerateMeshLODGraph.h b/Engine/Plugins/Experimental/MeshLODToolset/Source/MeshLODToolset/Public/Graphs/GenerateMeshLODGraph.h
--- a/Engine/Plugins/Experimental/MeshLODToolset/Source/MeshLODToolset/Public/Graphs/GenerateMeshLODGraph.h
+++ b/Engine/Plugins/Experimental/MeshLODToolset/Source/MeshLODToolset/Public/Graphs/GenerateMeshLODGraph.h
@@ -15,7 +15,7 @@
 #include "MeshProcessingNodes/MeshDeleteTrianglesNode.h"
 #include "MeshProcessingNodes/MeshAutoGenerateUVsNode.h"
 #include "DataTypes/MeshImageBakingData.h"
-#include "PhysicsNodes/GenerateConvexHullsCollisionNode.h"
+#include "PhysicsNodes/GenerateSimpleCollisionNode.h"
 
 
 class FGenerateMeshLODGraph
@@ -43,9 +43,8 @@ public:
 	void UpdateBakeCacheSettings(const UE::GeometryFlow::FMeshMakeBakingCacheSettings& BakeCacheSettings);
 	const UE::GeometryFlow::FMeshMakeBakingCacheSettings& GetCurrentBakeCacheSettings() const { return CurrentBakeCacheSettings; }
 
-
-	void UpdateGenerateConvexCollisionSettings(const UE::GeometryFlow::FGenerateConvexHullsCollisionSettings& ConvexesSettings);
-	const UE::GeometryFlow::FGenerateConvexHullsCollisionSettings& GetCurrentGenerateConvexCollisionSettings() const { return CurrentGenerateConvexHullsSettings; }
+	void UpdateGenerateSimpleCollisionSettings(const UE::GeometryFlow::FGenerateSimpleCollisionSettings& SimpleCollisionSettings);
+	const UE::GeometryFlow::FGenerateSimpleCollisionSettings& GetCurrentGenerateSimpleCollisionSettings() const { return CurrentGenerateSimpleCollisionSettings; }
 
 
 
@@ -115,12 +114,12 @@ protected:
 
 	UE::GeometryFlow::FGraph::FHandle DecomposeMeshForCollisionNode;
 	
-	UE::GeometryFlow::FGraph::FHandle GenerateConvexesNode;
-	UE::GeometryFlow::FGraph::FHandle GenerateConvexesSettingsNode;
-	UE::GeometryFlow::FGenerateConvexHullsCollisionSettings CurrentGenerateConvexHullsSettings;
+	UE::GeometryFlow::FGraph::FHandle GenerateSimpleCollisionNode;
+	UE::GeometryFlow::FGraph::FHandle GenerateSimpleCollisionSettingsNode;
+	UE::GeometryFlow::FGenerateSimpleCollisionSettings CurrentGenerateSimpleCollisionSettings;
 
 	UE::GeometryFlow::FGraph::FHandle CollisionOutputNode;
 	UE::GeometryFlow::FGraph::FHandle MeshOutputNode;
 	UE::GeometryFlow::FGraph::FHandle TangentsOutputNode;
 
-};
\ No newline at end of file
+};
diff --git a/Engine/Plugins/Experimental/MeshLODToolset/Source/MeshLODToolset/Public/Graphs/GenerateStaticMeshLODProcess.h b/Engine/Plugins/Experimental/MeshLODToolset/Source/MeshLODToolset/Public/Graphs/GenerateStaticMeshLODProcess.h
--- a/Engine/Plugins/Experimental/MeshLODToolset/Source/MeshLODToolset/Public/Graphs/GenerateStaticMeshLODProcess.h
+++ b/Engine/Plugins/Experimental/MeshLODToolset/Source/MeshLODToolset/Public/Graphs/GenerateStaticMeshLODProcess.h
@@ -32,6 +32,16 @@ enum class EGenerateStaticMeshLODBakeResolution
 	Resolution8192 = 8192 UMETA(DisplayName = "8192 x 8192")
 };
 
+UENUM()
+enum class EGenerateStaticMeshSimpleCollisionGeometryType : uint8
+{
+	AlignedBoxes,
+	OrientedBoxes,
+	MinimalSpheres,
+	Capsules,
+	ConvexHulls
+};
+
 
 USTRUCT()
 struct FGenerateStaticMeshLODProcessSettings
@@ -80,14 +90,18 @@ struct FGenerateStaticMeshLODProcessSettings
 
 	// Convex Hull Settings
 
-	UPROPERTY(EditAnywhere, Category = ConvexCollision, meta = (DisplayName = "Convex Tri Count"))
+	UPROPERTY(EditAnywhere, Category = Collision, meta = (DisplayName = "Collision Type"))
+	EGenerateStaticMeshSimpleCollisionGeometryType CollisionType = EGenerateStaticMeshSimpleCollisionGeometryType::ConvexHulls;
+
+	UPROPERTY(EditAnywhere, Category = Collision, meta = (DisplayName = "Convex Tri Count", 
+														  EditCondition = "CollisionType == EGenerateStaticMeshSimpleCollisionGeometryType::ConvexHulls"))
 	int ConvexTriangleCount = 50;
 
-	UPROPERTY(EditAnywhere, Category = ConvexCollision)
+	UPROPERTY(EditAnywhere, Category = Collision, meta = (EditCondition = "CollisionType == EGenerateStaticMeshSimpleCollisionGeometryType::ConvexHulls"))
 	bool bPrefilterVertices = false;
 
 	/** Grid resolution (along the maximum-length axis) */
-	UPROPERTY(EditAnywhere, Category = ConvexCollision, meta = (EditCondition = "bPrefilterVertices", UIMin = 1, UIMax = 30))
+	UPROPERTY(EditAnywhere, Category = Collision, meta = (EditCondition = "CollisionType == EGenerateStaticMeshSimpleCollisionGeometryType::ConvexHulls && bPrefilterVertices", UIMin = 1, UIMax = 30))
 	int PrefilterGridResolution = 10;
 
 };
