Add support for multiple group layers to MakeTriangleSetsNode. Retrieve the IndexSets for a given group layer by passing in the layer's name.

#rnx

diff --git a/Engine/Plugins/Experimental/GeometryFlow/Source/GeometryFlowCore/Public/GeometryFlowCoreNodes.h b/Engine/Plugins/Experimental/GeometryFlow/Source/GeometryFlowCore/Public/GeometryFlowCoreNodes.h
--- a/Engine/Plugins/Experimental/GeometryFlow/Source/GeometryFlowCore/Public/GeometryFlowCoreNodes.h
+++ b/Engine/Plugins/Experimental/GeometryFlow/Source/GeometryFlowCore/Public/GeometryFlowCoreNodes.h
@@ -112,7 +112,7 @@ GEOMETRYFLOW_DECLARE_BASIC_TYPES(Float, float, (int)EDataTypes::Float)
 GEOMETRYFLOW_DECLARE_BASIC_TYPES(Double, double, (int)EDataTypes::Double)
 GEOMETRYFLOW_DECLARE_BASIC_TYPES(Vector3f, FVector3f, (int)EDataTypes::Vector3f)
 GEOMETRYFLOW_DECLARE_BASIC_TYPES(Vector3d, FVector3d, (int)EDataTypes::Vector3d)
-
+GEOMETRYFLOW_DECLARE_BASIC_TYPES(Name, FName, (int)EDataTypes::Name)
 
 
 }	// end namespace GeometryFlow
diff --git a/Engine/Plugins/Experimental/GeometryFlow/Source/GeometryFlowCore/Public/GeometryFlowTypes.h b/Engine/Plugins/Experimental/GeometryFlow/Source/GeometryFlowCore/Public/GeometryFlowTypes.h
--- a/Engine/Plugins/Experimental/GeometryFlow/Source/GeometryFlowCore/Public/GeometryFlowTypes.h
+++ b/Engine/Plugins/Experimental/GeometryFlow/Source/GeometryFlowCore/Public/GeometryFlowTypes.h
@@ -56,6 +56,8 @@ enum class EDataTypes
 	Vector3f = 20,
 	Vector3d = 21,
 
+	Name = 30,
+
 	BaseMeshProcessingTypes = 1000,
 	BaseMeshProcessingTypesEditor = 2000,
 
diff --git a/Engine/Plugins/Experimental/GeometryFlow/Source/GeometryFlowMeshProcessing/Public/MeshDecompositionNodes/MakeTriangleSetsNode.h b/Engine/Plugins/Experimental/GeometryFlow/Source/GeometryFlowMeshProcessing/Public/MeshDecompositionNodes/MakeTriangleSetsNode.h
--- a/Engine/Plugins/Experimental/GeometryFlow/Source/GeometryFlowMeshProcessing/Public/MeshDecompositionNodes/MakeTriangleSetsNode.h
+++ b/Engine/Plugins/Experimental/GeometryFlow/Source/GeometryFlowMeshProcessing/Public/MeshDecompositionNodes/MakeTriangleSetsNode.h
@@ -81,36 +81,35 @@ public:
 class FMakeTriangleSetsFromGroupsNode : public FMakeTriangleSetsFromMeshNode
 {
 public:
+	static const FString InParamGroupLayer() { return TEXT("GroupLayerName"); }
 	static const FString InParamIgnoreGroups() { return TEXT("IgnoreGroups"); }
 
-
 public:
 	FMakeTriangleSetsFromGroupsNode() : FMakeTriangleSetsFromMeshNode()
 	{
+		AddInput(InParamGroupLayer(), MakeUnique<TBasicNodeInput<FName, (int)EDataTypes::Name>>());
 		AddInput(InParamIgnoreGroups(), MakeBasicInput<FIndexSets>());
 	}
 
 	virtual void CheckAdditionalInputs(const FNamedDataMap& DatasIn, bool& bRecomputeRequired, bool& bAllInputsValid) override
 	{
+		FindAndUpdateInputForEvaluate(InParamGroupLayer(), DatasIn, bRecomputeRequired, bAllInputsValid);
 		FindAndUpdateInputForEvaluate(InParamIgnoreGroups(), DatasIn, bRecomputeRequired, bAllInputsValid);
 	}
 
+protected:
 
-public:
-	
-	virtual void ComputeIndexSets(const FNamedDataMap& DatasIn, const FDynamicMesh3& Mesh, FIndexSets& SetsOut) override
+	void ComputeIndexSetsForGroups(const FDynamicMesh3& Mesh,
+								   const TSet<int32>& IgnoreGroups,
+								   TFunction<int32(int)> TriangleGroupFn,
+								   FIndexSets& SetsOut)
 	{
-		TSafeSharedPtr<IData> IgnoreGroupsArg = DatasIn.FindData(InParamIgnoreGroups());
-		const FIndexSets& IgnoreGroupsSets = IgnoreGroupsArg->GetDataConstRef<FIndexSets>(FIndexSets::DataTypeIdentifier);
-		TSet<int32> IgnoreGroups;
-		IgnoreGroupsSets.GetAllValues(IgnoreGroups);
-
 		TMap<int32, int32> GroupsMap;
 		TArray<int32> GroupCounts;
 		int32 NumGroups = 0;
 		for (int32 tid : Mesh.TriangleIndicesItr())
 		{
-			int32 GroupID = Mesh.GetTriangleGroup(tid);
+			int32 GroupID = TriangleGroupFn(tid);
 			if (IgnoreGroups.Contains(GroupID))
 			{
 				continue;
@@ -137,7 +136,7 @@ public:
 
 		for (int32 tid : Mesh.TriangleIndicesItr())
 		{
-			int32 GroupID = Mesh.GetTriangleGroup(tid);
+			int32 GroupID = TriangleGroupFn(tid);
 			if (IgnoreGroups.Contains(GroupID))
 			{
 				continue;
@@ -148,6 +147,36 @@ public:
 		}
 	}
 
+public:
+
+	virtual void ComputeIndexSets(const FNamedDataMap& DatasIn, const FDynamicMesh3& Mesh, FIndexSets& SetsOut) override
+	{
+		TSafeSharedPtr<IData> IgnoreGroupsArg = DatasIn.FindData(InParamIgnoreGroups());
+		const FIndexSets& IgnoreGroupsSets = IgnoreGroupsArg->GetDataConstRef<FIndexSets>(FIndexSets::DataTypeIdentifier);
+		TSet<int32> IgnoreGroups;
+		IgnoreGroupsSets.GetAllValues(IgnoreGroups);
+
+		TSafeSharedPtr<IData> GroupLayerArg = DatasIn.FindData(InParamGroupLayer());
+		FName GroupName = GroupLayerArg->GetDataConstRef<FName>((int)EDataTypes::Name);
+
+		if (Mesh.HasTriangleGroups() && (GroupName == "Default"))
+		{
+			ComputeIndexSetsForGroups(Mesh, IgnoreGroups, [&Mesh](int TriangleID) { return Mesh.GetTriangleGroup(TriangleID); }, SetsOut);
+		}
+		else if (Mesh.HasAttributes())
+		{
+			for (int32 PolygroupLayerIndex = 0; PolygroupLayerIndex < Mesh.Attributes()->NumPolygroupLayers(); ++PolygroupLayerIndex)
+			{
+				const FDynamicMeshPolygroupAttribute* PolygroupLayer = Mesh.Attributes()->GetPolygroupLayer(PolygroupLayerIndex);
+				if (PolygroupLayer->GetName() == GroupName)
+				{
+					ComputeIndexSetsForGroups(Mesh, IgnoreGroups, [&PolygroupLayer](int TriangleID) { return PolygroupLayer->GetValue(TriangleID); }, SetsOut);
+				}
+			}		
+		}
+
+	}
+
 };
 
 
