Enable default values on Node inputs. AddInput now takes an additional optional parameter for the default value.

#rnx

diff --git a/Engine/Plugins/Experimental/GeometryFlow/Source/GeometryFlowCore/Private/GeometryFlowExecutor.cpp b/Engine/Plugins/Experimental/GeometryFlow/Source/GeometryFlowCore/Private/GeometryFlowExecutor.cpp
--- a/Engine/Plugins/Experimental/GeometryFlow/Source/GeometryFlowCore/Private/GeometryFlowExecutor.cpp
+++ b/Engine/Plugins/Experimental/GeometryFlow/Source/GeometryFlowCore/Private/GeometryFlowExecutor.cpp
@@ -45,7 +45,7 @@ void FGeometryFlowExecutor::TopologicalSort()
 	TopologicallySortedNodes.Empty();
 
 	// S := Set of all nodes with no incoming edge
-	TSet<FGraph::FHandle> CurrentNodes = GeometryFlowGraph.GetSourceNodes();
+	TSet<FGraph::FHandle> CurrentNodes = GeometryFlowGraph.GetNodesWithNoConnectedInputs();
 
 	while (CurrentNodes.Num() != 0)
 	{
@@ -124,7 +124,16 @@ void FGeometryFlowExecutor::AsyncRunGraph()
 				// find the connection for this input
 				FGraph::FConnection Connection;
 				EGeometryFlowResult FoundResult = GeometryFlowGraph.FindConnectionForInput(NodeHandle, InputName, Connection);
-				check(FoundResult == EGeometryFlowResult::Ok);
+
+				if (FoundResult == EGeometryFlowResult::ConnectionDoesNotExist)
+				{			
+					TSafeSharedPtr<IData> DefaultData = Node->GetDefaultInputData(InputName);
+					// TODO: Bubble this error up rather than crash
+					checkf(DefaultData != nullptr, TEXT("Node \"%s\" input \"%s\" is not connected and has no default value"), *Node->GetIdentifier(), *InputName);
+					FDataFlags DataFlags;
+					DatasIn.Add(InputName, DefaultData, DataFlags);
+					return;
+				}
 
 				// If there is only one Connection from this upstream Output (ie to our Input), and the Node/Input 
 				// can steal and transform that data, then we will do it
diff --git a/Engine/Plugins/Experimental/GeometryFlow/Source/GeometryFlowCore/Private/GeometryFlowGraph.cpp b/Engine/Plugins/Experimental/GeometryFlow/Source/GeometryFlowCore/Private/GeometryFlowGraph.cpp
--- a/Engine/Plugins/Experimental/GeometryFlow/Source/GeometryFlowCore/Private/GeometryFlowGraph.cpp
+++ b/Engine/Plugins/Experimental/GeometryFlow/Source/GeometryFlowCore/Private/GeometryFlowGraph.cpp
@@ -117,6 +117,47 @@ EGeometryFlowResult FGraph::InferConnection(FHandle FromNodeHandle, FHandle ToNo
 	}
 }
 
+TSet<FGraph::FHandle> FGraph::GetSourceNodes() const
+{
+	TSet<FHandle> SourceNodes;
+	for (const TMap<FHandle, FNodeInfo>::ElementType& IdNodePair : AllNodes)
+	{
+		if (IdNodePair.Value.Node->NodeInputs.Num() == 0)
+		{
+			SourceNodes.Add(FHandle{ IdNodePair.Key });
+		}
+	}
+	return SourceNodes;
+}
+
+TSet<FGraph::FHandle> FGraph::GetNodesWithNoConnectedInputs() const
+{
+	TSet<FHandle> NoInputNodes;
+
+	for (const TMap<FHandle, FNodeInfo>::ElementType& IdNodePair : AllNodes)
+	{
+		FHandle NodeHandle = IdNodePair.Key;
+		TSafeSharedPtr<FNode> Node = IdNodePair.Value.Node;
+
+		bool bAnyInputConnected = false;
+		Node->EnumerateInputs([this, NodeHandle, &bAnyInputConnected](const FString& InputName, const TUniquePtr<INodeInput>& Input)
+		{
+			FConnection ConnectionOut;
+			EGeometryFlowResult FindResult = FindConnectionForInput(NodeHandle, InputName, ConnectionOut);
+			if (FindResult == EGeometryFlowResult::Ok)
+			{
+				bAnyInputConnected = true;
+			}
+		});
+
+		if (!bAnyInputConnected)
+		{
+			NoInputNodes.Add(NodeHandle);
+		}
+	}
+
+	return NoInputNodes;
+}
 
 
 
@@ -131,7 +172,6 @@ EGeometryFlowResult FGraph::FindConnectionForInput(FHandle ToNode, FString ToInp
 			return EGeometryFlowResult::Ok;
 		}
 	}
-	ensure(false);
 	return EGeometryFlowResult::ConnectionDoesNotExist;
 }
 
@@ -179,7 +219,15 @@ TSafeSharedPtr<IData> FGraph::ComputeOutputData(
 		// find the connection for this input
 		FConnection Connection;
 		EGeometryFlowResult FoundResult = FindConnectionForInput(NodeHandle, InputName, Connection);
-		check(FoundResult == EGeometryFlowResult::Ok);
+
+		if (FoundResult == EGeometryFlowResult::ConnectionDoesNotExist) 
+		{
+			TSafeSharedPtr<IData> DefaultData = Node->GetDefaultInputData(InputName);
+			// TODO: Bubble this error up rather than crash
+			checkf(DefaultData != nullptr, TEXT("Node \"%s\" input \"%s\" is not connected and has no default value"), *Node->GetIdentifier(), *InputName);
+			DataIn.Add(InputName, DefaultData, DataFlags);
+			continue;
+		}
 
 		ENodeCachingStrategy FromCachingStrategy = GetCachingStrategyForNode(Connection.FromNode);
 
diff --git a/Engine/Plugins/Experimental/GeometryFlow/Source/GeometryFlowCore/Private/GeometryFlowNode.cpp b/Engine/Plugins/Experimental/GeometryFlow/Source/GeometryFlowCore/Private/GeometryFlowNode.cpp
--- a/Engine/Plugins/Experimental/GeometryFlow/Source/GeometryFlowCore/Private/GeometryFlowNode.cpp
+++ b/Engine/Plugins/Experimental/GeometryFlow/Source/GeometryFlowCore/Private/GeometryFlowNode.cpp
@@ -5,9 +5,13 @@
 using namespace UE::GeometryFlow;
 
 
-void FNode::AddInput(FString Name, TUniquePtr<INodeInput>&& Input)
+void FNode::AddInput(FString Name, TUniquePtr<INodeInput>&& Input, TSafeSharedPtr<IData> DefaultData)
 {
 	// todo checks/etc
+	if (DefaultData)
+	{
+		InputDefaultValues.Add(Name, DefaultData);
+	}
 
 	FNodeInputInfo NodeInputInfo;
 	NodeInputInfo.Name = Name;
diff --git a/Engine/Plugins/Experimental/GeometryFlow/Source/GeometryFlowCore/Public/GeometryFlowGraph.h b/Engine/Plugins/Experimental/GeometryFlow/Source/GeometryFlowCore/Public/GeometryFlowGraph.h
--- a/Engine/Plugins/Experimental/GeometryFlow/Source/GeometryFlowCore/Public/GeometryFlowGraph.h
+++ b/Engine/Plugins/Experimental/GeometryFlow/Source/GeometryFlowCore/Public/GeometryFlowGraph.h
@@ -70,19 +70,8 @@ public:
 
 	EGeometryFlowResult InferConnection(FHandle FromNode, FHandle ToNode);
 
-	inline TSet<FHandle> GetSourceNodes() const
-	{
-		TSet<FHandle> SourceNodes;
-		for (const TMap<FHandle, FNodeInfo>::ElementType& IdNodePair : AllNodes)
-		{
-			if (IdNodePair.Value.Node->NodeInputs.Num() == 0)
-			{
-				SourceNodes.Add(FHandle{ IdNodePair.Key });
-			}
-		}
-		return SourceNodes;
-	}
-
+	TSet<FHandle> GetSourceNodes() const;
+	TSet<FHandle> GetNodesWithNoConnectedInputs() const;
 
 	template<typename T>
 	EGeometryFlowResult EvaluateResult(FHandle Node, FString OutputName, T& Storage, int32 StorageTypeIdentifier)
diff --git a/Engine/Plugins/Experimental/GeometryFlow/Source/GeometryFlowCore/Public/GeometryFlowNode.h b/Engine/Plugins/Experimental/GeometryFlow/Source/GeometryFlowCore/Public/GeometryFlowNode.h
--- a/Engine/Plugins/Experimental/GeometryFlow/Source/GeometryFlowCore/Public/GeometryFlowNode.h
+++ b/Engine/Plugins/Experimental/GeometryFlow/Source/GeometryFlowCore/Public/GeometryFlowNode.h
@@ -290,7 +290,7 @@ public:
 	void SetIdentifier(const FString& IdentifierIn) { Identifier = IdentifierIn; }
 	const FString& GetIdentifier() const { return Identifier; }
 
-	void AddInput(FString Name, TUniquePtr<INodeInput>&& Input);
+	void AddInput(FString Name, TUniquePtr<INodeInput>&& Input, TSafeSharedPtr<IData> DefaultData = TSafeSharedPtr<IData>());
 	void AddOutput(FString Name, TUniquePtr<INodeOutput>&& Output);
 
 	EGeometryFlowResult GetInputType(FString Name, int32& TypeOut) const;
@@ -322,6 +322,15 @@ public:
 		FNamedDataMap& DatasOut,
 		TUniquePtr<FEvaluationInfo>& EvaluationInfo) = 0;
 
+	virtual TSafeSharedPtr<IData> GetDefaultInputData(FString InputName) const
+	{
+		if (InputDefaultValues.Contains(InputName))
+		{
+			return InputDefaultValues[InputName];
+		}
+		return TSafeSharedPtr<IData>();
+	}
+
 protected:
 
 	friend class FGraph;
@@ -364,6 +373,7 @@ protected:
 	TArray<FNodeInputInfo> NodeInputs;
 	TArray<FNodeOutputInfo> NodeOutputs;
 
+	TMap<FString, TSafeSharedPtr<IData>> InputDefaultValues;
 
 protected:
 	// helper functions for setup
diff --git a/Engine/Plugins/Experimental/GeometryFlow/Source/GeometryFlowMeshProcessing/Public/MeshDecompositionNodes/MakeTriangleSetsNode.h b/Engine/Plugins/Experimental/GeometryFlow/Source/GeometryFlowMeshProcessing/Public/MeshDecompositionNodes/MakeTriangleSetsNode.h
--- a/Engine/Plugins/Experimental/GeometryFlow/Source/GeometryFlowMeshProcessing/Public/MeshDecompositionNodes/MakeTriangleSetsNode.h
+++ b/Engine/Plugins/Experimental/GeometryFlow/Source/GeometryFlowMeshProcessing/Public/MeshDecompositionNodes/MakeTriangleSetsNode.h
@@ -39,6 +39,8 @@ public:
 			bool bAllInputsValid = true;
 			bool bRecomputeRequired = (IsOutputAvailable(OutParamIndexSets()) == false);
 			TSafeSharedPtr<IData> MeshArg = FindAndUpdateInputForEvaluate(InParam(), DatasIn, bRecomputeRequired, bAllInputsValid);
+			CheckAdditionalInputs(DatasIn, bRecomputeRequired, bAllInputsValid);
+
 			if (bAllInputsValid)
 			{
 				if (bRecomputeRequired)
@@ -87,8 +89,12 @@ public:
 public:
 	FMakeTriangleSetsFromGroupsNode() : FMakeTriangleSetsFromMeshNode()
 	{
-		AddInput(InParamGroupLayer(), MakeUnique<TBasicNodeInput<FName, (int)EDataTypes::Name>>());
-		AddInput(InParamIgnoreGroups(), MakeBasicInput<FIndexSets>());
+		AddInput(InParamGroupLayer(), 
+				 MakeUnique<TBasicNodeInput<FName, (int)EDataTypes::Name>>(), 
+				 MakeSafeShared<TMovableData<FName, (int)EDataTypes::Name>>("Default"));
+
+		AddInput(InParamIgnoreGroups(), 
+				 MakeBasicInput<FIndexSets>());
 	}
 
 	virtual void CheckAdditionalInputs(const FNamedDataMap& DatasIn, bool& bRecomputeRequired, bool& bAllInputsValid) override
diff --git a/Engine/Plugins/Experimental/MeshLODToolset/Source/MeshLODToolset/Private/Graphs/GenerateMeshLODGraph.cpp b/Engine/Plugins/Experimental/MeshLODToolset/Source/MeshLODToolset/Private/Graphs/GenerateMeshLODGraph.cpp
--- a/Engine/Plugins/Experimental/MeshLODToolset/Source/MeshLODToolset/Private/Graphs/GenerateMeshLODGraph.cpp
+++ b/Engine/Plugins/Experimental/MeshLODToolset/Source/MeshLODToolset/Private/Graphs/GenerateMeshLODGraph.cpp
@@ -402,13 +402,11 @@ void FGenerateMeshLODGraph::BuildGraph()
 	// collision generation
 
 	FGraph::FHandle IgnoreGroupsForCollisionNode = Graph->AddNodeOfType<FIndexSetsSourceNode>(TEXT("CollisionIgnoreGroups"));
-	FGraph::FHandle GroupLayerNameNode = Graph->AddNodeOfType<FNameSourceNode>(TEXT("GroupLayerNameSource"));
 
 	//DecomposeMeshForCollisionNode = Graph->AddNodeOfType<FMakeTriangleSetsFromMeshNode>(TEXT("Decompose"));
 	DecomposeMeshForCollisionNode = Graph->AddNodeOfType<FMakeTriangleSetsFromGroupsNode>(TEXT("Decompose"));
 	Graph->InferConnection(MeshSourceNode, DecomposeMeshForCollisionNode);
 	Graph->InferConnection(IgnoreGroupsForCollisionNode, DecomposeMeshForCollisionNode);
-	Graph->InferConnection(GroupLayerNameNode, DecomposeMeshForCollisionNode);
 
 	GenerateSimpleCollisionNode = Graph->AddNodeOfType<FGenerateSimpleCollisionNode>(TEXT("GenerateSimpleCollision"));
 	Graph->InferConnection(MeshSourceNode, GenerateSimpleCollisionNode);
@@ -475,9 +473,6 @@ void FGenerateMeshLODGraph::BuildGraph()
 	IgnoreGroupsForCollision.AppendSet({ 0 });
 	UpdateSettingsSourceNodeValue(*Graph, IgnoreGroupsForCollisionNode, IgnoreGroupsForCollision);
 
-	FName GroupLayerName{ "Default" };
-	UpdateSourceNodeValue<FNameSourceNode>(*Graph, GroupLayerNameNode, GroupLayerName);
-
 	FGenerateSimpleCollisionSettings GenSimpleCollisionSettings;
 	UpdateGenerateSimpleCollisionSettings(GenSimpleCollisionSettings);
 
