Fix GeometryFlow unit tests after adding ProgressCancel and removing DebugNodeExecutionLog from the Executor in CL 15273200.

#rnx

diff --git a/Engine/Plugins/Experimental/GeometryFlow/Source/GeometryFlowCore/Public/GeometryFlowExecutor.h b/Engine/Plugins/Experimental/GeometryFlow/Source/GeometryFlowCore/Public/GeometryFlowExecutor.h
--- a/Engine/Plugins/Experimental/GeometryFlow/Source/GeometryFlowCore/Public/GeometryFlowExecutor.h
+++ b/Engine/Plugins/Experimental/GeometryFlow/Source/GeometryFlowCore/Public/GeometryFlowExecutor.h
@@ -38,7 +38,7 @@ public:
 	// corresponding output data!
 	EGeometryFlowResult ComputeOutputs(const TArray<NodeOutputSpec>& DesiredOutputs,
 									   TArray<TSafeSharedPtr<IData>>& OutputDatas,
-									   FProgressCancel* Progress);
+									   FProgressCancel* Progress = nullptr);
 
 	// 2a.
 	// Launch graph execution on background threads and return
