Alternate parallel execution for GeometryFlow graphs. Starting from the desired output, recursively kick off upstream nodes with Async. Wait for the upstream nodes to finish and then execute the node.

#rnx

diff --git a/Engine/Plugins/Experimental/GeometryFlow/Source/GeometryFlowCore/Private/GeometryFlowGraph.cpp b/Engine/Plugins/Experimental/GeometryFlow/Source/GeometryFlowCore/Private/GeometryFlowGraph.cpp
--- a/Engine/Plugins/Experimental/GeometryFlow/Source/GeometryFlowCore/Private/GeometryFlowGraph.cpp
+++ b/Engine/Plugins/Experimental/GeometryFlow/Source/GeometryFlowCore/Private/GeometryFlowGraph.cpp
@@ -2,10 +2,19 @@
 
 #include "GeometryFlowGraph.h"
 #include "Util/ProgressCancel.h"
+#include "Async/Async.h"
 
 using namespace UE::GeometryFlow;
 
 
+namespace
+{
+	// probably should be something defined for the whole tool framework. UETOOL-2989.
+	static EAsyncExecution GeometryFlowGraphAsyncExecTarget = EAsyncExecution::Thread;
+
+	// Can't seem to use threadpool in GeometryProcessingUnitTests. GThreadPool is null.
+}
+
 TSafeSharedPtr<FNode> FGraph::FindNode(FHandle NodeHandle) const
 {
 	const FNodeInfo* Found = AllNodes.Find(NodeHandle);
@@ -16,6 +25,16 @@ TSafeSharedPtr<FNode> FGraph::FindNode(FHandle NodeHandle) const
 	return Found->Node;
 }
 
+TSafeSharedPtr<FRWLock> FGraph::FindNodeLock(FHandle NodeHandle) const
+{
+	const TSafeSharedPtr<FRWLock>* Found = AllNodeLocks.Find(NodeHandle);
+	if (!ensure(Found != nullptr))
+	{
+		return nullptr;
+	}
+	return *Found;
+}
+
 EGeometryFlowResult FGraph::GetInputTypeForNode(FHandle NodeHandle, FString InputName, int32& Type) const
 {
 	const FNodeInfo* Found = AllNodes.Find(NodeHandle);
@@ -210,6 +229,9 @@ TSafeSharedPtr<IData> FGraph::ComputeOutputData(
 	// this is the map of (InputName, Data) we will build up by pulling from the Connections
 	FNamedDataMap DataIn;
 
+	TArray<TFuture<void>> AsyncFutures;
+	FRWLock DataInLock;
+
 	// Collect data from those Inputs.
 	// This will recursively call ComputeOutputData() on those (Node/Output) pairs
 	for (int32 k = 0; k < InputRequirements.Num(); ++k)
@@ -226,7 +248,9 @@ TSafeSharedPtr<IData> FGraph::ComputeOutputData(
 			TSafeSharedPtr<IData> DefaultData = Node->GetDefaultInputData(InputName);
 			// TODO: Bubble this error up rather than crash
 			checkf(DefaultData != nullptr, TEXT("Node \"%s\" input \"%s\" is not connected and has no default value"), *Node->GetIdentifier(), *InputName);
+			DataInLock.WriteLock();
 			DataIn.Add(InputName, DefaultData, DataFlags);
+			DataInLock.WriteUnlock();
 			continue;
 		}
 
@@ -244,14 +268,27 @@ TSafeSharedPtr<IData> FGraph::ComputeOutputData(
 			DataFlags.bIsMutableData = true;
 		}
 
-		// recursively fetch Data coming in to this Input via Connection
-		TSafeSharedPtr<IData> UpstreamData = 
-			ComputeOutputData(Connection.FromNode, Connection.FromOutput, EvaluationInfo, bStealDataForInput);
+		TFuture<void> Future = Async(GeometryFlowGraphAsyncExecTarget, 
+									[this, &DataIn, &DataInLock, &InputName, DataFlags, Connection, &EvaluationInfo, bStealDataForInput] ()
+		{
+			// recursively fetch Data coming in to this Input via Connection
+			TSafeSharedPtr<IData> UpstreamData = ComputeOutputData(Connection.FromNode, Connection.FromOutput, EvaluationInfo, bStealDataForInput);
+
+			DataInLock.WriteLock();
+			DataIn.Add(InputName, UpstreamData, DataFlags);
+			DataInLock.WriteUnlock();
+		});
+
+		AsyncFutures.Emplace(MoveTemp(Future));
+	}
 
-		// add to named data map
-		DataIn.Add(InputName, UpstreamData, DataFlags);
+	for (TFuture<void>& Future : AsyncFutures)
+	{
+		Future.Wait();
 	}
 
+	check(DataIn.GetNames().Num() == InputRequirements.Num());
+
 	if (EvaluationInfo && EvaluationInfo->Progress && EvaluationInfo->Progress->Cancelled())
 	{
 		return nullptr;
@@ -260,7 +297,11 @@ TSafeSharedPtr<IData> FGraph::ComputeOutputData(
 	// evalute node
 	FNamedDataMap DataOut;
 	DataOut.Add(OutputName);
+	
+	FindNodeLock(NodeHandle)->WriteLock();
 	Node->Evaluate(DataIn, DataOut, EvaluationInfo);
+	FindNodeLock(NodeHandle)->WriteUnlock();
+
 	EvaluationInfo->CountEvaluation(Node.Get());
 
 	if (EvaluationInfo && EvaluationInfo->Progress && EvaluationInfo->Progress->Cancelled())
diff --git a/Engine/Plugins/Experimental/GeometryFlow/Source/GeometryFlowCore/Public/GeometryFlowGraph.h b/Engine/Plugins/Experimental/GeometryFlow/Source/GeometryFlowCore/Public/GeometryFlowGraph.h
--- a/Engine/Plugins/Experimental/GeometryFlow/Source/GeometryFlowCore/Public/GeometryFlowGraph.h
+++ b/Engine/Plugins/Experimental/GeometryFlow/Source/GeometryFlowCore/Public/GeometryFlowGraph.h
@@ -66,6 +66,8 @@ public:
 		NewNodeInfo.CachingStrategy = CachingStrategy;
 		FHandle Handle = { NodeCounter++ };
 		AllNodes.Add(Handle, NewNodeInfo);
+
+		AllNodeLocks.Add(Handle, MakeSafeShared<FRWLock>());
 		return Handle;
 	}
 
@@ -166,11 +168,13 @@ protected:
 	};
 
 	TMap<FHandle, FNodeInfo> AllNodes;
+	TMap<FHandle, TSafeSharedPtr<FRWLock>> AllNodeLocks;
 
 	TSafeSharedPtr<FNode> FindNode(FHandle Handle) const;
 	EGeometryFlowResult GetInputTypeForNode(FHandle NodeHandle, FString InputName, int32& Type) const;
 	EGeometryFlowResult GetOutputTypeForNode(FHandle NodeHandle, FString OutputName, int32& Type) const;
 	ENodeCachingStrategy GetCachingStrategyForNode(FHandle NodeHandle) const;
+	TSafeSharedPtr<FRWLock> FindNodeLock(FHandle Handle) const;
 
 	TArray<FConnection> Connections;
 
