GeometryFlowGraph: detect when trying to connect to an input that already has an incoming connection.

#rnx

diff --git a/Engine/Plugins/Experimental/GeometryFlow/Source/GeometryFlowCore/Private/GeometryFlowGraph.cpp b/Engine/Plugins/Experimental/GeometryFlow/Source/GeometryFlowCore/Private/GeometryFlowGraph.cpp
--- a/Engine/Plugins/Experimental/GeometryFlow/Source/GeometryFlowCore/Private/GeometryFlowGraph.cpp
+++ b/Engine/Plugins/Experimental/GeometryFlow/Source/GeometryFlowCore/Private/GeometryFlowGraph.cpp
@@ -3,6 +3,7 @@
 #include "GeometryFlowGraph.h"
 #include "Util/ProgressCancel.h"
 #include "Async/Async.h"
+#include "Algo/Find.h"
 
 using namespace UE::GeometryFlow;
 
@@ -89,7 +90,14 @@ EGeometryFlowResult FGraph::AddConnection(FHandle FromNode, FString FromOutput,
 		return EGeometryFlowResult::UnmatchedTypes;
 	}
 
-	// TODO: check for unique inputs
+	FConnection* Found = Algo::FindByPredicate(Connections, [&](const FConnection& Conn)
+	{
+		return ((Conn.ToNode == ToNode) && (Conn.ToInput == ToInput));
+	});
+	if (Found != nullptr)
+	{
+		return EGeometryFlowResult::InputAlreadyConnected;
+	}
 
 	Connections.Add({ FromNode, FromOutput, ToNode, ToInput });
 
diff --git a/Engine/Plugins/Experimental/GeometryFlow/Source/GeometryFlowCore/Public/GeometryFlowTypes.h b/Engine/Plugins/Experimental/GeometryFlow/Source/GeometryFlowCore/Public/GeometryFlowTypes.h
--- a/Engine/Plugins/Experimental/GeometryFlow/Source/GeometryFlowCore/Public/GeometryFlowTypes.h
+++ b/Engine/Plugins/Experimental/GeometryFlow/Source/GeometryFlowCore/Public/GeometryFlowTypes.h
@@ -40,7 +40,8 @@ enum class EGeometryFlowResult
 	ConnectionDoesNotExist = 5,
 	NoMatchesFound = 6,
 	MultipleMatchingAmbiguityFound = 7,
-	OperationCancelled = 8
+	OperationCancelled = 8,
+	InputAlreadyConnected = 9
 };
 
 
