Make convex hull cancellable when called from GeometryFlow.

#rnx

diff --git a/Engine/Plugins/Experimental/GeometryFlow/Source/GeometryFlowMeshProcessing/Private/PhysicsNodes/GenerateSimpleCollisionNode.cpp b/Engine/Plugins/Experimental/GeometryFlow/Source/GeometryFlowMeshProcessing/Private/PhysicsNodes/GenerateSimpleCollisionNode.cpp
--- a/Engine/Plugins/Experimental/GeometryFlow/Source/GeometryFlowMeshProcessing/Private/PhysicsNodes/GenerateSimpleCollisionNode.cpp
+++ b/Engine/Plugins/Experimental/GeometryFlow/Source/GeometryFlowMeshProcessing/Private/PhysicsNodes/GenerateSimpleCollisionNode.cpp
@@ -23,7 +23,8 @@ namespace GenerateSimpleCollisionNodeHelpers
 	void GenerateConvexHulls(const FDynamicMesh3& Mesh,
 							 const FIndexSets& IndexData,
 							 const FGenerateSimpleCollisionSettings& Settings,
-							 FCollisionGeometry& Result)
+							 FCollisionGeometry& Result,
+							 FProgressCancel* Progress)
 	{
 		int32 NumShapes = IndexData.IndexSets.Num();
 		TArray<FConvexShape3d> Convexes;
@@ -44,7 +45,7 @@ namespace GenerateSimpleCollisionNodeHelpers
 
 			Hull.bPostSimplify = Settings.ConvexHullSettings.SimplifyToTriangleCount > 0;
 			Hull.MaxTargetFaceCount = Settings.ConvexHullSettings.SimplifyToTriangleCount;
-			if (Hull.Compute())
+			if (Hull.Compute(Progress))
 			{
 				FConvexShape3d NewConvex;
 				NewConvex.Mesh = MoveTemp(Hull.ConvexHull);
@@ -77,7 +78,7 @@ namespace GenerateSimpleCollisionNodeHelpers
 
 				Hull.bPostSimplify = Settings.ConvexHullSettings.SimplifyToTriangleCount > 0;
 				Hull.MaxTargetFaceCount = Settings.ConvexHullSettings.SimplifyToTriangleCount;
-				if (Hull.Compute())
+				if (Hull.Compute(Progress))
 				{
 					FConvexShape3d NewConvex;
 					NewConvex.Mesh = MoveTemp(Hull.ConvexHull);
@@ -164,7 +165,7 @@ EGeometryFlowResult FGenerateSimpleCollisionNode::EvaluateInternal(const FDynami
 		 ShapeApproximator.Generate_Capsules(OutCollisionGeometry.Geometry);
 		 break;
 	 case ESimpleCollisionGeometryType::ConvexHulls:
-		 GenerateSimpleCollisionNodeHelpers::GenerateConvexHulls(Mesh, IndexData, Settings, OutCollisionGeometry);
+		 GenerateSimpleCollisionNodeHelpers::GenerateConvexHulls(Mesh, IndexData, Settings, OutCollisionGeometry, Progress);
 		 break;
 	 case ESimpleCollisionGeometryType::SweptHulls:
 		 ShapeApproximator.bDetectConvexes = true;
diff --git a/Engine/Plugins/Experimental/GeometryProcessing/Source/DynamicMesh/Private/Operations/MeshConvexHull.cpp b/Engine/Plugins/Experimental/GeometryProcessing/Source/DynamicMesh/Private/Operations/MeshConvexHull.cpp
--- a/Engine/Plugins/Experimental/GeometryProcessing/Source/DynamicMesh/Private/Operations/MeshConvexHull.cpp
+++ b/Engine/Plugins/Experimental/GeometryProcessing/Source/DynamicMesh/Private/Operations/MeshConvexHull.cpp
@@ -10,16 +10,16 @@
 #include "ExplicitUseGeometryMathTypes.h"		// using UE::Geometry::(math types)
 using namespace UE::Geometry;
 
-bool FMeshConvexHull::Compute()
+bool FMeshConvexHull::Compute(FProgressCancel* Progress)
 {
 	bool bOK = false;
 	if (VertexSet.Num() > 0)
 	{
-		bOK = Compute_VertexSubset();
+		bOK = Compute_VertexSubset(Progress);
 	}
 	else
 	{
-		bOK = Compute_FullMesh();
+		bOK = Compute_FullMesh(Progress);
 	}
 	if (!bOK)
 	{
@@ -45,7 +45,7 @@ bool FMeshConvexHull::Compute()
 			// recalculate convex hull
 			// TODO: test if simplified mesh is convex first, can just re-use in that case!!
 			FMeshConvexHull SimplifiedHull(&ConvexHull);
-			if (SimplifiedHull.Compute())
+			if (SimplifiedHull.Compute(Progress))
 			{
 				ConvexHull = MoveTemp(SimplifiedHull.ConvexHull);
 			}
@@ -60,9 +60,10 @@ bool FMeshConvexHull::Compute()
 
 
 
-bool FMeshConvexHull::Compute_FullMesh()
+bool FMeshConvexHull::Compute_FullMesh(FProgressCancel* Progress)
 {
 	FConvexHull3d HullCompute;
+	HullCompute.Progress = Progress;
 	bool bOK = HullCompute.Solve(Mesh->MaxVertexID(),
 		[this](int32 Index) { return Mesh->GetVertex(Index); },
 		[this](int32 Index) { return Mesh->IsVertex(Index); });
@@ -100,9 +101,10 @@ bool FMeshConvexHull::Compute_FullMesh()
 
 
 
-bool FMeshConvexHull::Compute_VertexSubset()
+bool FMeshConvexHull::Compute_VertexSubset(FProgressCancel* Progress)
 {
 	FConvexHull3d HullCompute;
+	HullCompute.Progress = Progress;
 	bool bOK = HullCompute.Solve(VertexSet.Num(),
 		[this](int32 Index) { return Mesh->GetVertex(VertexSet[Index]); });
 	if (!bOK)
diff --git a/Engine/Plugins/Experimental/GeometryProcessing/Source/DynamicMesh/Private/ShapeApproximation/MeshSimpleShapeApproximation.cpp b/Engine/Plugins/Experimental/GeometryProcessing/Source/DynamicMesh/Private/ShapeApproximation/MeshSimpleShapeApproximation.cpp
--- a/Engine/Plugins/Experimental/GeometryProcessing/Source/DynamicMesh/Private/ShapeApproximation/MeshSimpleShapeApproximation.cpp
+++ b/Engine/Plugins/Experimental/GeometryProcessing/Source/DynamicMesh/Private/ShapeApproximation/MeshSimpleShapeApproximation.cpp
@@ -160,7 +160,7 @@ static void ComputeSimpleShapeFits(const FDynamicMesh3& Mesh,
 	if (bConvex)
 	{
 		FMeshConvexHull Hull(&Mesh);
-		if (Hull.Compute())
+		if (Hull.Compute(Progress))
 		{
 			FitResult.bHaveConvex = true;
 			FitResult.Convex = MoveTemp(Hull.ConvexHull);
@@ -271,7 +271,7 @@ void FMeshSimpleShapeApproximation::Generate_Capsules(FSimpleShapeSet3d& ShapeSe
 
 
 
-void FMeshSimpleShapeApproximation::Generate_ConvexHulls(FSimpleShapeSet3d& ShapeSetOut)
+void FMeshSimpleShapeApproximation::Generate_ConvexHulls(FSimpleShapeSet3d& ShapeSetOut, FProgressCancel* Progress)
 {
 	FCriticalSection GeometryLock;
 	ParallelFor(SourceMeshes.Num(), [&](int32 idx)
@@ -287,7 +287,7 @@ void FMeshSimpleShapeApproximation::Generate_ConvexHulls(FSimpleShapeSet3d& Shap
 		Hull.bPostSimplify = bSimplifyHulls;
 		Hull.MaxTargetFaceCount = HullTargetFaceCount;
 
-		if (Hull.Compute())
+		if (Hull.Compute(Progress))
 		{
 			FConvexShape3d NewConvex;
 			NewConvex.Mesh = MoveTemp(Hull.ConvexHull);
diff --git a/Engine/Plugins/Experimental/GeometryProcessing/Source/DynamicMesh/Public/Operations/MeshConvexHull.h b/Engine/Plugins/Experimental/GeometryProcessing/Source/DynamicMesh/Public/Operations/MeshConvexHull.h
--- a/Engine/Plugins/Experimental/GeometryProcessing/Source/DynamicMesh/Public/Operations/MeshConvexHull.h
+++ b/Engine/Plugins/Experimental/GeometryProcessing/Source/DynamicMesh/Public/Operations/MeshConvexHull.h
@@ -8,6 +8,8 @@
 #include "GeometryTypes.h"
 #include "DynamicMesh3.h"
 
+class FProgressCancel;
+
 namespace UE
 {
 namespace Geometry
@@ -62,13 +64,13 @@ public:
 	 * Calculate output ConvexHull mesh for vertices of input Mesh
 	 * @return true on success
 	 */
-	bool Compute();
+	bool Compute(FProgressCancel* Progress = nullptr);
 
 
 
 protected:
-	bool Compute_FullMesh();
-	bool Compute_VertexSubset();
+	bool Compute_FullMesh(FProgressCancel* Progress);
+	bool Compute_VertexSubset(FProgressCancel* Progress);
 };
 
 
diff --git a/Engine/Plugins/Experimental/GeometryProcessing/Source/DynamicMesh/Public/ShapeApproximation/MeshSimpleShapeApproximation.h b/Engine/Plugins/Experimental/GeometryProcessing/Source/DynamicMesh/Public/ShapeApproximation/MeshSimpleShapeApproximation.h
--- a/Engine/Plugins/Experimental/GeometryProcessing/Source/DynamicMesh/Public/ShapeApproximation/MeshSimpleShapeApproximation.h
+++ b/Engine/Plugins/Experimental/GeometryProcessing/Source/DynamicMesh/Public/ShapeApproximation/MeshSimpleShapeApproximation.h
@@ -113,7 +113,7 @@ public:
 	 * Calculate 3D Convex Hulls for each input mesh and store in ShapeSetOut.
 	 * Each convex hull is stored as a triangle mesh, and optionally simplified if bSimplifyHulls=true
 	 */
-	void Generate_ConvexHulls(FSimpleShapeSet3d& ShapeSetOut);
+	void Generate_ConvexHulls(FSimpleShapeSet3d& ShapeSetOut, FProgressCancel* Progress = nullptr);
 
 
 	/** Type/Mode for deciding 3D axis to use in Generate_ProjectedHulls() */
