GeometryFlow: Resolve PVS-Studio warnings in FMeshMakeCleanGeometry

diff --git a/Engine/Plugins/Experimental/GeometryFlow/Source/GeometryFlowMeshProcessing/Private/MeshProcessingNodes/MeshMakeCleanGeometryNode.cpp b/Engine/Plugins/Experimental/GeometryFlow/Source/GeometryFlowMeshProcessing/Private/MeshProcessingNodes/MeshMakeCleanGeometryNode.cpp
--- a/Engine/Plugins/Experimental/GeometryFlow/Source/GeometryFlowMeshProcessing/Private/MeshProcessingNodes/MeshMakeCleanGeometryNode.cpp
+++ b/Engine/Plugins/Experimental/GeometryFlow/Source/GeometryFlowMeshProcessing/Private/MeshProcessingNodes/MeshMakeCleanGeometryNode.cpp
@@ -84,10 +84,10 @@ void FMeshMakeCleanGeometryNode::ApplyMakeCleanGeometry(FDynamicMesh3& Mesh, con
 		}
 	}
 
-	if (Settings.FillHolesEdgeCountThresh > 0 || Settings.FillHolesEstimatedAreaFraction > 0)
+	if (Settings.FillHolesEdgeCountThresh > 0 || Settings.FillHolesEstimatedAreaFraction > 0.0)
 	{
 		double MeshArea = 0.0;
-		if (Settings.FillHolesEstimatedAreaFraction > 0)
+		if (Settings.FillHolesEstimatedAreaFraction > 0.0)
 		{
 			FVector2d VolArea = TMeshQueries<FDynamicMesh3>::GetVolumeArea(Mesh);
 			MeshArea = VolArea.Y;
