AutoLOD: Pass EvaluationInfo to node ProcessMesh functions. Make GenerateUVs and GenerateConvexHullMesh cancellable.

#rnx

diff --git a/Engine/Plugins/Experimental/GeometryFlow/Source/GeometryFlowMeshProcessing/Private/MeshProcessingNodes/GenerateConvexHullMeshNode.cpp b/Engine/Plugins/Experimental/GeometryFlow/Source/GeometryFlowMeshProcessing/Private/MeshProcessingNodes/GenerateConvexHullMeshNode.cpp
--- a/Engine/Plugins/Experimental/GeometryFlow/Source/GeometryFlowMeshProcessing/Private/MeshProcessingNodes/GenerateConvexHullMeshNode.cpp
+++ b/Engine/Plugins/Experimental/GeometryFlow/Source/GeometryFlowMeshProcessing/Private/MeshProcessingNodes/GenerateConvexHullMeshNode.cpp
@@ -9,7 +9,8 @@ using namespace UE::GeometryFlow;
 
 EGeometryFlowResult FGenerateConvexHullMeshNode::MakeConvexHullMesh(const FDynamicMesh3& MeshIn,
 	const FGenerateConvexHullMeshSettings& Settings, 
-	FDynamicMesh3& MeshOut)
+	FDynamicMesh3& MeshOut,
+	TUniquePtr<FEvaluationInfo>& EvaluationInfo)
 {
 	FMeshConvexHull Hull(&MeshIn);
 
@@ -20,12 +21,18 @@ EGeometryFlowResult FGenerateConvexHullMeshNode::MakeConvexHullMesh(const FDynam
 
 	Hull.bPostSimplify = false;		// Mesh can be simplified later
 
-	FProgressCancel* Progress = nullptr;	// TODO: Add ProgressCancel
-	if (Hull.Compute(Progress))
+	if (Hull.Compute(EvaluationInfo->Progress))
 	{
 		MeshOut = MoveTemp(Hull.ConvexHull);
 	}
 
+	if (EvaluationInfo->Progress->Cancelled())
+	{
+		return EGeometryFlowResult::OperationCancelled;
+	}
+
+	// TODO: What if Hull.Compute() fails but not because of cancellation? Should we have a EGeometryFlowResult::Failed?
+
 	return EGeometryFlowResult::Ok;
 }
 
diff --git a/Engine/Plugins/Experimental/GeometryFlow/Source/GeometryFlowMeshProcessing/Private/MeshProcessingNodes/MeshMakeCleanGeometryNode.cpp b/Engine/Plugins/Experimental/GeometryFlow/Source/GeometryFlowMeshProcessing/Private/MeshProcessingNodes/MeshMakeCleanGeometryNode.cpp
--- a/Engine/Plugins/Experimental/GeometryFlow/Source/GeometryFlowMeshProcessing/Private/MeshProcessingNodes/MeshMakeCleanGeometryNode.cpp
+++ b/Engine/Plugins/Experimental/GeometryFlow/Source/GeometryFlowMeshProcessing/Private/MeshProcessingNodes/MeshMakeCleanGeometryNode.cpp
@@ -16,7 +16,8 @@ void FMeshMakeCleanGeometryNode::ProcessMesh(
 	const FNamedDataMap& DatasIn,
 	const FMeshMakeCleanGeometrySettings& SettingsIn,
 	const FDynamicMesh3& MeshIn,
-	FDynamicMesh3& MeshOut) 
+	FDynamicMesh3& MeshOut,
+	TUniquePtr<FEvaluationInfo>& EvaluationInfo)
 {
 	MeshOut = MeshIn;
 	ApplyMakeCleanGeometry(MeshOut, SettingsIn);
@@ -26,7 +27,8 @@ void FMeshMakeCleanGeometryNode::ProcessMesh(
 void FMeshMakeCleanGeometryNode::ProcessMeshInPlace(
 	const FNamedDataMap& DatasIn,
 	const FMeshMakeCleanGeometrySettings& Settings,
-	FDynamicMesh3& MeshInOut) 
+	FDynamicMesh3& MeshInOut,
+	TUniquePtr<FEvaluationInfo>& EvaluationInfo)
 {
 	ApplyMakeCleanGeometry(MeshInOut, Settings);
 }
diff --git a/Engine/Plugins/Experimental/GeometryFlow/Source/GeometryFlowMeshProcessing/Private/MeshProcessingNodes/MeshNormalFlowNode.cpp b/Engine/Plugins/Experimental/GeometryFlow/Source/GeometryFlowMeshProcessing/Private/MeshProcessingNodes/MeshNormalFlowNode.cpp
--- a/Engine/Plugins/Experimental/GeometryFlow/Source/GeometryFlowMeshProcessing/Private/MeshProcessingNodes/MeshNormalFlowNode.cpp
+++ b/Engine/Plugins/Experimental/GeometryFlow/Source/GeometryFlowMeshProcessing/Private/MeshProcessingNodes/MeshNormalFlowNode.cpp
@@ -97,7 +97,8 @@ void FMeshNormalFlowNode::ProcessMesh(
 	const FNamedDataMap& DatasIn,
 	const FMeshNormalFlowSettings& SettingsIn,
 	const FDynamicMesh3& MeshIn,
-	FDynamicMesh3& MeshOut)
+	FDynamicMesh3& MeshOut,
+	TUniquePtr<FEvaluationInfo>& EvaluationInfo)
 {
 	TSafeSharedPtr<IData> ProjectionTargetMeshArg = DatasIn.FindData(InParamTargetMesh());
 	FDynamicMesh3 ProjectionTargetMesh;
@@ -113,7 +114,8 @@ void FMeshNormalFlowNode::ProcessMesh(
 void FMeshNormalFlowNode::ProcessMeshInPlace(
 	const FNamedDataMap& DatasIn,
 	const FMeshNormalFlowSettings& SettingsIn,
-	FDynamicMesh3& MeshInOut)
+	FDynamicMesh3& MeshInOut,
+	TUniquePtr<FEvaluationInfo>& EvaluationInfo)
 {
 	// TODO: In this case, can we call IData::GiveTo instead of GetDataCopy for the projection target?
 	TSafeSharedPtr<IData> ProjectionTargetMeshArg = DatasIn.FindData(InParamTargetMesh());
diff --git a/Engine/Plugins/Experimental/GeometryFlow/Source/GeometryFlowMeshProcessing/Private/MeshProcessingNodes/MeshThickenNode.cpp b/Engine/Plugins/Experimental/GeometryFlow/Source/GeometryFlowMeshProcessing/Private/MeshProcessingNodes/MeshThickenNode.cpp
--- a/Engine/Plugins/Experimental/GeometryFlow/Source/GeometryFlowMeshProcessing/Private/MeshProcessingNodes/MeshThickenNode.cpp
+++ b/Engine/Plugins/Experimental/GeometryFlow/Source/GeometryFlowMeshProcessing/Private/MeshProcessingNodes/MeshThickenNode.cpp
@@ -12,7 +12,8 @@ void FMeshThickenNode::ProcessMesh(
 	const FNamedDataMap& DatasIn,
 	const FMeshThickenSettings& SettingsIn,
 	const FDynamicMesh3& MeshIn,
-	FDynamicMesh3& MeshOut) 
+	FDynamicMesh3& MeshOut,
+	TUniquePtr<FEvaluationInfo>& EvaluationInfo)
 {
 	TSafeSharedPtr<IData> WeightMapMeshArg = DatasIn.FindData(InParamWeightMap());
 	FWeightMap WeightMap;
@@ -26,7 +27,8 @@ void FMeshThickenNode::ProcessMesh(
 void FMeshThickenNode::ProcessMeshInPlace(
 	const FNamedDataMap& DatasIn,
 	const FMeshThickenSettings& Settings,
-	FDynamicMesh3& MeshInOut) 
+	FDynamicMesh3& MeshInOut,
+	TUniquePtr<FEvaluationInfo>& EvaluationInfo)
 {
 	TSafeSharedPtr<IData> WeightMapMeshArg = DatasIn.FindData(InParamWeightMap());
 	FWeightMap WeightMap;
diff --git a/Engine/Plugins/Experimental/GeometryFlow/Source/GeometryFlowMeshProcessing/Public/MeshProcessingNodes/CompactMeshNode.h b/Engine/Plugins/Experimental/GeometryFlow/Source/GeometryFlowMeshProcessing/Public/MeshProcessingNodes/CompactMeshNode.h
--- a/Engine/Plugins/Experimental/GeometryFlow/Source/GeometryFlowMeshProcessing/Public/MeshProcessingNodes/CompactMeshNode.h
+++ b/Engine/Plugins/Experimental/GeometryFlow/Source/GeometryFlowMeshProcessing/Public/MeshProcessingNodes/CompactMeshNode.h
@@ -26,14 +26,16 @@ public:
 	virtual void ProcessMesh(
 		const FNamedDataMap& DatasIn,
 		const FDynamicMesh3& MeshIn,
-		FDynamicMesh3& MeshOut) override
+		FDynamicMesh3& MeshOut,
+		TUniquePtr<FEvaluationInfo>& EvaluationInfo) override
 	{
 		MeshOut.CompactCopy(MeshIn);
 	}
 
 	virtual void ProcessMeshInPlace(
 		const FNamedDataMap& DatasIn,
-		FDynamicMesh3& MeshInOut) override
+		FDynamicMesh3& MeshInOut,
+		TUniquePtr<FEvaluationInfo>& EvaluationInfo) override
 	{
 		MeshInOut.CompactInPlace();
 	}
diff --git a/Engine/Plugins/Experimental/GeometryFlow/Source/GeometryFlowMeshProcessing/Public/MeshProcessingNodes/GenerateConvexHullMeshNode.h b/Engine/Plugins/Experimental/GeometryFlow/Source/GeometryFlowMeshProcessing/Public/MeshProcessingNodes/GenerateConvexHullMeshNode.h
--- a/Engine/Plugins/Experimental/GeometryFlow/Source/GeometryFlowMeshProcessing/Public/MeshProcessingNodes/GenerateConvexHullMeshNode.h
+++ b/Engine/Plugins/Experimental/GeometryFlow/Source/GeometryFlowMeshProcessing/Public/MeshProcessingNodes/GenerateConvexHullMeshNode.h
@@ -31,12 +31,16 @@ public:
 		const FNamedDataMap& DatasIn,
 		const FGenerateConvexHullMeshSettings& SettingsIn,
 		const FDynamicMesh3& MeshIn,
-		FDynamicMesh3& MeshOut) override
+		FDynamicMesh3& MeshOut,
+		TUniquePtr<FEvaluationInfo>& EvaluationInfo) override
 	{
-		MakeConvexHullMesh(MeshIn, SettingsIn, MeshOut);
+		MakeConvexHullMesh(MeshIn, SettingsIn, MeshOut, EvaluationInfo);
 	}
 
-	EGeometryFlowResult MakeConvexHullMesh(const FDynamicMesh3& MeshIn, const FGenerateConvexHullMeshSettings& Settings, FDynamicMesh3& MeshOut);
+	EGeometryFlowResult MakeConvexHullMesh(const FDynamicMesh3& MeshIn,
+		const FGenerateConvexHullMeshSettings& Settings,
+		FDynamicMesh3& MeshOut,
+		TUniquePtr<FEvaluationInfo>& EvaluationInfo);
 
 };
 
diff --git a/Engine/Plugins/Experimental/GeometryFlow/Source/GeometryFlowMeshProcessing/Public/MeshProcessingNodes/MeshDeleteTrianglesNode.h b/Engine/Plugins/Experimental/GeometryFlow/Source/GeometryFlowMeshProcessing/Public/MeshProcessingNodes/MeshDeleteTrianglesNode.h
--- a/Engine/Plugins/Experimental/GeometryFlow/Source/GeometryFlowMeshProcessing/Public/MeshProcessingNodes/MeshDeleteTrianglesNode.h
+++ b/Engine/Plugins/Experimental/GeometryFlow/Source/GeometryFlowMeshProcessing/Public/MeshProcessingNodes/MeshDeleteTrianglesNode.h
@@ -37,7 +37,8 @@ public:
 	virtual void ProcessMesh(
 		const FNamedDataMap& DatasIn,
 		const FDynamicMesh3& MeshIn,
-		FDynamicMesh3& MeshOut) override
+		FDynamicMesh3& MeshOut,
+		TUniquePtr<FEvaluationInfo>& EvaluationInfo) override
 	{
 		MeshOut = MeshIn;
 		DeleteTrianglesForMesh(DatasIn, MeshOut);
@@ -45,7 +46,8 @@ public:
 
 	virtual void ProcessMeshInPlace(
 		const FNamedDataMap& DatasIn,
-		FDynamicMesh3& MeshInOut) override
+		FDynamicMesh3& MeshInOut,
+		TUniquePtr<FEvaluationInfo>& EvaluationInfo) override
 	{
 		DeleteTrianglesForMesh(DatasIn, MeshInOut);
 	}
diff --git a/Engine/Plugins/Experimental/GeometryFlow/Source/GeometryFlowMeshProcessing/Public/MeshProcessingNodes/MeshMakeCleanGeometryNode.h b/Engine/Plugins/Experimental/GeometryFlow/Source/GeometryFlowMeshProcessing/Public/MeshProcessingNodes/MeshMakeCleanGeometryNode.h
--- a/Engine/Plugins/Experimental/GeometryFlow/Source/GeometryFlowMeshProcessing/Public/MeshProcessingNodes/MeshMakeCleanGeometryNode.h
+++ b/Engine/Plugins/Experimental/GeometryFlow/Source/GeometryFlowMeshProcessing/Public/MeshProcessingNodes/MeshMakeCleanGeometryNode.h
@@ -38,12 +38,14 @@ public:
 		const FNamedDataMap& DatasIn,
 		const FMeshMakeCleanGeometrySettings& SettingsIn,
 		const FDynamicMesh3& MeshIn,
-		FDynamicMesh3& MeshOut) override;
+		FDynamicMesh3& MeshOut,
+		TUniquePtr<FEvaluationInfo>& EvaluationInfo) override;
 
 	void ProcessMeshInPlace(
 		const FNamedDataMap& DatasIn,
 		const FMeshMakeCleanGeometrySettings& Settings,
-		FDynamicMesh3& MeshInOut) override;
+		FDynamicMesh3& MeshInOut,
+		TUniquePtr<FEvaluationInfo>& EvaluationInfo) override;
 
 protected:
 
diff --git a/Engine/Plugins/Experimental/GeometryFlow/Source/GeometryFlowMeshProcessing/Public/MeshProcessingNodes/MeshNormalFlowNode.h b/Engine/Plugins/Experimental/GeometryFlow/Source/GeometryFlowMeshProcessing/Public/MeshProcessingNodes/MeshNormalFlowNode.h
--- a/Engine/Plugins/Experimental/GeometryFlow/Source/GeometryFlowMeshProcessing/Public/MeshProcessingNodes/MeshNormalFlowNode.h
+++ b/Engine/Plugins/Experimental/GeometryFlow/Source/GeometryFlowMeshProcessing/Public/MeshProcessingNodes/MeshNormalFlowNode.h
@@ -46,12 +46,14 @@ public:
 		const FNamedDataMap& DatasIn,
 		const FMeshNormalFlowSettings& SettingsIn,
 		const FDynamicMesh3& MeshIn,
-		FDynamicMesh3& MeshOut) override;
+		FDynamicMesh3& MeshOut,
+		TUniquePtr<FEvaluationInfo>& EvaluationInfo) override;
 
 	void ProcessMeshInPlace(
 		const FNamedDataMap& DatasIn,
 		const FMeshNormalFlowSettings& SettingsIn,
-		FDynamicMesh3& MeshInOut) override;
+		FDynamicMesh3& MeshInOut,
+		TUniquePtr<FEvaluationInfo>& EvaluationInfo) override;
 
 	void CheckAdditionalInputs(const FNamedDataMap& DatasIn, 
 							   bool& bRecomputeRequired, 
diff --git a/Engine/Plugins/Experimental/GeometryFlow/Source/GeometryFlowMeshProcessing/Public/MeshProcessingNodes/MeshNormalsNodes.h b/Engine/Plugins/Experimental/GeometryFlow/Source/GeometryFlowMeshProcessing/Public/MeshProcessingNodes/MeshNormalsNodes.h
--- a/Engine/Plugins/Experimental/GeometryFlow/Source/GeometryFlowMeshProcessing/Public/MeshProcessingNodes/MeshNormalsNodes.h
+++ b/Engine/Plugins/Experimental/GeometryFlow/Source/GeometryFlowMeshProcessing/Public/MeshProcessingNodes/MeshNormalsNodes.h
@@ -58,7 +58,8 @@ public:
 		const FNamedDataMap& DatasIn,
 		const FMeshNormalsSettings& Settings,
 		const FDynamicMesh3& MeshIn,
-		FDynamicMesh3& MeshOut) override
+		FDynamicMesh3& MeshOut,
+		TUniquePtr<FEvaluationInfo>& EvaluationInfo) override
 	{
 		MeshOut = MeshIn;
 		ComputeNormals(Settings, MeshOut);
@@ -67,7 +68,8 @@ public:
 	virtual void ProcessMeshInPlace(
 		const FNamedDataMap& DatasIn,
 		const FMeshNormalsSettings& Settings,
-		FDynamicMesh3& MeshInOut)
+		FDynamicMesh3& MeshInOut,
+		TUniquePtr<FEvaluationInfo>& EvaluationInfo)
 	{
 		ComputeNormals(Settings, MeshInOut);
 	}
@@ -128,7 +130,7 @@ public:
 class FComputeMeshPerVertexOverlayNormalsNode : public FSimpleInPlaceProcessMeshBaseNode
 {
 public:
-	virtual void ApplyNodeToMesh(FDynamicMesh3& MeshInOut)
+	virtual void ApplyNodeToMesh(FDynamicMesh3& MeshInOut, TUniquePtr<FEvaluationInfo>& EvaluationInfo) override
 	{
 		if (MeshInOut.HasAttributes() == false)
 		{
@@ -152,7 +154,7 @@ public:
 class FComputeMeshPerVertexNormalsNode : public FSimpleInPlaceProcessMeshBaseNode
 {
 public:
-	virtual void ApplyNodeToMesh(FDynamicMesh3& MeshInOut)
+	virtual void ApplyNodeToMesh(FDynamicMesh3& MeshInOut, TUniquePtr<FEvaluationInfo>& EvaluationInfo) override
 	{
 		FMeshNormals::QuickComputeVertexNormals(MeshInOut, false);
 	}
diff --git a/Engine/Plugins/Experimental/GeometryFlow/Source/GeometryFlowMeshProcessing/Public/MeshProcessingNodes/MeshProcessingBaseNodes.h b/Engine/Plugins/Experimental/GeometryFlow/Source/GeometryFlowMeshProcessing/Public/MeshProcessingNodes/MeshProcessingBaseNodes.h
--- a/Engine/Plugins/Experimental/GeometryFlow/Source/GeometryFlowMeshProcessing/Public/MeshProcessingNodes/MeshProcessingBaseNodes.h
+++ b/Engine/Plugins/Experimental/GeometryFlow/Source/GeometryFlowMeshProcessing/Public/MeshProcessingNodes/MeshProcessingBaseNodes.h
@@ -40,14 +40,16 @@ protected:
 	virtual void ProcessMesh(
 		const FNamedDataMap& DatasIn,
 		const FDynamicMesh3& MeshIn,
-		FDynamicMesh3& MeshOut)
+		FDynamicMesh3& MeshOut,
+		TUniquePtr<FEvaluationInfo>& EvaluationInfo)
 	{
 		check(false);		// subclass must implement ProcessMesh()
 	}
 
 	virtual void ProcessMeshInPlace(
 		const FNamedDataMap& DatasIn,
-		FDynamicMesh3& MeshInOut)
+		FDynamicMesh3& MeshInOut,
+		TUniquePtr<FEvaluationInfo>& EvaluationInfo)
 	{
 		ensureMsgf(false, TEXT("TProcessMeshBaseNode::ProcessMeshInPlace called but not defined!"));
 	}
@@ -82,7 +84,7 @@ public:
 
 						FDynamicMesh3 EditableMesh;
 						MeshArg->GiveTo<FDynamicMesh3>(EditableMesh, (int)EMeshProcessingDataTypes::DynamicMesh);
-						ProcessMeshInPlace(DatasIn, EditableMesh);
+						ProcessMeshInPlace(DatasIn, EditableMesh, EvaluationInfo);
 
 						// store new result
 						TSafeSharedPtr<FDataDynamicMesh> Result = MakeSafeShared<FDataDynamicMesh>();
@@ -98,7 +100,7 @@ public:
 
 						// run mesh processing
 						FDynamicMesh3 ResultMesh;
-						ProcessMesh(DatasIn, SourceMesh, ResultMesh);
+						ProcessMesh(DatasIn, SourceMesh, ResultMesh, EvaluationInfo);
 
 						// store new result
 						TSafeSharedPtr<FDataDynamicMesh> Result = MakeSafeShared<FDataDynamicMesh>();
@@ -133,19 +135,19 @@ public:
 		ConfigureInputFlags(InParamMesh(), FNodeInputFlags::Transformable());
 	}
 
-	virtual void ProcessMesh( const FNamedDataMap& DatasIn, const FDynamicMesh3& MeshIn, FDynamicMesh3& MeshOut) override
+	virtual void ProcessMesh( const FNamedDataMap& DatasIn, const FDynamicMesh3& MeshIn, FDynamicMesh3& MeshOut, TUniquePtr<FEvaluationInfo>& EvaluationInfo) override
 	{
 		MeshOut = MeshIn;
-		ApplyNodeToMesh(MeshOut);
+		ApplyNodeToMesh(MeshOut, EvaluationInfo);
 	}
 
-	virtual void ProcessMeshInPlace( const FNamedDataMap& DatasIn, FDynamicMesh3& MeshInOut) override
+	virtual void ProcessMeshInPlace( const FNamedDataMap& DatasIn, FDynamicMesh3& MeshInOut, TUniquePtr<FEvaluationInfo>& EvaluationInfo) override
 	{
-		ApplyNodeToMesh(MeshInOut);
+		ApplyNodeToMesh(MeshInOut, EvaluationInfo);
 	}
 
 	// subclasses only have to implement this
-	virtual void ApplyNodeToMesh(FDynamicMesh3& MeshInOut) = 0;
+	virtual void ApplyNodeToMesh(FDynamicMesh3& MeshInOut, TUniquePtr<FEvaluationInfo>& EvaluationInfo) = 0;
 
 };
 
@@ -178,14 +180,16 @@ protected:
 		const FNamedDataMap& DatasIn,
 		const SettingsType& SettingsIn,
 		const FDynamicMesh3& MeshIn,
-		FDynamicMesh3& MeshOut)
+		FDynamicMesh3& MeshOut,
+		TUniquePtr<FEvaluationInfo>& EvaluationInfo)
 	{
 	}
 
 	virtual void ProcessMeshInPlace(
 		const FNamedDataMap& DatasIn,
 		const SettingsType& SettingsIn,
-		FDynamicMesh3& MeshInOut)
+		FDynamicMesh3& MeshInOut,
+		TUniquePtr<FEvaluationInfo>& EvaluationInfo)
 	{
 		ensureMsgf(false, TEXT("TProcessMeshWithSettingsBaseNode::ProcessMeshInPlace called but not defined!"));
 	}
@@ -225,7 +229,7 @@ public:
 
 						FDynamicMesh3 EditableMesh;
 						MeshArg->GiveTo<FDynamicMesh3>(EditableMesh, (int)EMeshProcessingDataTypes::DynamicMesh);
-						ProcessMeshInPlace(DatasIn, Settings, EditableMesh);
+						ProcessMeshInPlace(DatasIn, Settings, EditableMesh, EvaluationInfo);
 
 						// store new result
 						TSafeSharedPtr<FDataDynamicMesh> Result = MakeSafeShared<FDataDynamicMesh>();
@@ -241,7 +245,7 @@ public:
 
 						// run mesh processing
 						FDynamicMesh3 ResultMesh;
-						ProcessMesh(DatasIn, Settings, SourceMesh, ResultMesh);
+						ProcessMesh(DatasIn, Settings, SourceMesh, ResultMesh, EvaluationInfo);
 
 						// store new result
 						TSafeSharedPtr<FDataDynamicMesh> Result = MakeSafeShared<FDataDynamicMesh>();
diff --git a/Engine/Plugins/Experimental/GeometryFlow/Source/GeometryFlowMeshProcessing/Public/MeshProcessingNodes/MeshRecalculateUVsNode.h b/Engine/Plugins/Experimental/GeometryFlow/Source/GeometryFlowMeshProcessing/Public/MeshProcessingNodes/MeshRecalculateUVsNode.h
--- a/Engine/Plugins/Experimental/GeometryFlow/Source/GeometryFlowMeshProcessing/Public/MeshProcessingNodes/MeshRecalculateUVsNode.h
+++ b/Engine/Plugins/Experimental/GeometryFlow/Source/GeometryFlowMeshProcessing/Public/MeshProcessingNodes/MeshRecalculateUVsNode.h
@@ -43,7 +43,8 @@ public:
 		const FNamedDataMap& DatasIn,
 		const FMeshRecalculateUVsSettings& Settings,
 		const FDynamicMesh3& MeshIn,
-		FDynamicMesh3& MeshOut) override
+		FDynamicMesh3& MeshOut,
+		TUniquePtr<FEvaluationInfo>& EvaluationInfo) override
 	{
 		MeshOut = MeshIn;
 		RecalculateUVsOnMesh(MeshOut, Settings);
@@ -52,7 +53,8 @@ public:
 	virtual void ProcessMeshInPlace(
 		const FNamedDataMap& DatasIn,
 		const FMeshRecalculateUVsSettings& Settings,
-		FDynamicMesh3& MeshInOut)
+		FDynamicMesh3& MeshInOut,
+		TUniquePtr<FEvaluationInfo>& EvaluationInfo)
 	{
 		RecalculateUVsOnMesh(MeshInOut, Settings);
 	}
diff --git a/Engine/Plugins/Experimental/GeometryFlow/Source/GeometryFlowMeshProcessing/Public/MeshProcessingNodes/MeshRepackUVsNode.h b/Engine/Plugins/Experimental/GeometryFlow/Source/GeometryFlowMeshProcessing/Public/MeshProcessingNodes/MeshRepackUVsNode.h
--- a/Engine/Plugins/Experimental/GeometryFlow/Source/GeometryFlowMeshProcessing/Public/MeshProcessingNodes/MeshRepackUVsNode.h
+++ b/Engine/Plugins/Experimental/GeometryFlow/Source/GeometryFlowMeshProcessing/Public/MeshProcessingNodes/MeshRepackUVsNode.h
@@ -42,7 +42,8 @@ public:
 		const FNamedDataMap& DatasIn,
 		const FMeshRepackUVsSettings& Settings,
 		const FDynamicMesh3& MeshIn,
-		FDynamicMesh3& MeshOut) override
+		FDynamicMesh3& MeshOut,
+		TUniquePtr<FEvaluationInfo>& EvaluationInfo) override
 	{
 		MeshOut = MeshIn;
 		RepackUVsForMesh(MeshOut, Settings);
@@ -51,7 +52,8 @@ public:
 	virtual void ProcessMeshInPlace(
 		const FNamedDataMap& DatasIn,
 		const FMeshRepackUVsSettings& Settings,
-		FDynamicMesh3& MeshInOut)
+		FDynamicMesh3& MeshInOut,
+		TUniquePtr<FEvaluationInfo>& EvaluationInfo)
 	{
 		RepackUVsForMesh(MeshInOut, Settings);
 	}
diff --git a/Engine/Plugins/Experimental/GeometryFlow/Source/GeometryFlowMeshProcessing/Public/MeshProcessingNodes/MeshSimplifyNode.h b/Engine/Plugins/Experimental/GeometryFlow/Source/GeometryFlowMeshProcessing/Public/MeshProcessingNodes/MeshSimplifyNode.h
--- a/Engine/Plugins/Experimental/GeometryFlow/Source/GeometryFlowMeshProcessing/Public/MeshProcessingNodes/MeshSimplifyNode.h
+++ b/Engine/Plugins/Experimental/GeometryFlow/Source/GeometryFlowMeshProcessing/Public/MeshProcessingNodes/MeshSimplifyNode.h
@@ -70,7 +70,8 @@ public:
 		const FNamedDataMap& DatasIn,
 		const FMeshSimplifySettings& Settings,
 		const FDynamicMesh3& MeshIn,
-		FDynamicMesh3& MeshOut) override
+		FDynamicMesh3& MeshOut,
+		TUniquePtr<FEvaluationInfo>& EvaluationInfo) override
 	{
 		MeshOut.Copy(MeshIn, true, true, true, !Settings.bDiscardAttributes);
 		ApplySimplify(Settings, MeshOut);
@@ -79,7 +80,8 @@ public:
 	virtual void ProcessMeshInPlace(
 		const FNamedDataMap& DatasIn,
 		const FMeshSimplifySettings& Settings,
-		FDynamicMesh3& MeshInOut)
+		FDynamicMesh3& MeshInOut,
+		TUniquePtr<FEvaluationInfo>& EvaluationInfo)
 	{
 		ApplySimplify(Settings, MeshInOut);
 	}
diff --git a/Engine/Plugins/Experimental/GeometryFlow/Source/GeometryFlowMeshProcessing/Public/MeshProcessingNodes/MeshSolidifyNode.h b/Engine/Plugins/Experimental/GeometryFlow/Source/GeometryFlowMeshProcessing/Public/MeshProcessingNodes/MeshSolidifyNode.h
--- a/Engine/Plugins/Experimental/GeometryFlow/Source/GeometryFlowMeshProcessing/Public/MeshProcessingNodes/MeshSolidifyNode.h
+++ b/Engine/Plugins/Experimental/GeometryFlow/Source/GeometryFlowMeshProcessing/Public/MeshProcessingNodes/MeshSolidifyNode.h
@@ -35,7 +35,8 @@ public:
 		const FNamedDataMap& DatasIn,
 		const FMeshSolidifySettings& SettingsIn,
 		const FDynamicMesh3& MeshIn,
-		FDynamicMesh3& MeshOut) override
+		FDynamicMesh3& MeshOut,
+		TUniquePtr<FEvaluationInfo>& EvaluationInfo) override
 	{
 		ApplySolidify(MeshIn, SettingsIn, MeshOut);
 	}
diff --git a/Engine/Plugins/Experimental/GeometryFlow/Source/GeometryFlowMeshProcessing/Public/MeshProcessingNodes/MeshThickenNode.h b/Engine/Plugins/Experimental/GeometryFlow/Source/GeometryFlowMeshProcessing/Public/MeshProcessingNodes/MeshThickenNode.h
--- a/Engine/Plugins/Experimental/GeometryFlow/Source/GeometryFlowMeshProcessing/Public/MeshProcessingNodes/MeshThickenNode.h
+++ b/Engine/Plugins/Experimental/GeometryFlow/Source/GeometryFlowMeshProcessing/Public/MeshProcessingNodes/MeshThickenNode.h
@@ -42,12 +42,14 @@ public:
 		const FNamedDataMap& DatasIn,
 		const FMeshThickenSettings& SettingsIn,
 		const FDynamicMesh3& MeshIn,
-		FDynamicMesh3& MeshOut) override;
+		FDynamicMesh3& MeshOut,
+		TUniquePtr<FEvaluationInfo>& EvaluationInfo) override;
 
 	void ProcessMeshInPlace(
 		const FNamedDataMap& DatasIn,
 		const FMeshThickenSettings& Settings,
-		FDynamicMesh3& MeshInOut) override;
+		FDynamicMesh3& MeshInOut,
+		TUniquePtr<FEvaluationInfo>& EvaluationInfo) override;
 
 protected:
 
diff --git a/Engine/Plugins/Experimental/GeometryFlow/Source/GeometryFlowMeshProcessing/Public/MeshProcessingNodes/MeshVoxMorphologyNode.h b/Engine/Plugins/Experimental/GeometryFlow/Source/GeometryFlowMeshProcessing/Public/MeshProcessingNodes/MeshVoxMorphologyNode.h
--- a/Engine/Plugins/Experimental/GeometryFlow/Source/GeometryFlowMeshProcessing/Public/MeshProcessingNodes/MeshVoxMorphologyNode.h
+++ b/Engine/Plugins/Experimental/GeometryFlow/Source/GeometryFlowMeshProcessing/Public/MeshProcessingNodes/MeshVoxMorphologyNode.h
@@ -41,7 +41,8 @@ public:
 		const FNamedDataMap& DatasIn,
 		const FVoxMorphologyOpSettings& SettingsIn,
 		const FDynamicMesh3& MeshIn,
-		FDynamicMesh3& MeshOut) override
+		FDynamicMesh3& MeshOut,
+		TUniquePtr<FEvaluationInfo>& EvaluationInfo) override
 	{
 		ApplyMorphology(MeshIn, SettingsIn, MeshOut);
 	}
diff --git a/Engine/Plugins/Experimental/GeometryFlow/Source/GeometryFlowMeshProcessing/Public/MeshProcessingNodes/TransferMeshMaterialIDsNode.h b/Engine/Plugins/Experimental/GeometryFlow/Source/GeometryFlowMeshProcessing/Public/MeshProcessingNodes/TransferMeshMaterialIDsNode.h
--- a/Engine/Plugins/Experimental/GeometryFlow/Source/GeometryFlowMeshProcessing/Public/MeshProcessingNodes/TransferMeshMaterialIDsNode.h
+++ b/Engine/Plugins/Experimental/GeometryFlow/Source/GeometryFlowMeshProcessing/Public/MeshProcessingNodes/TransferMeshMaterialIDsNode.h
@@ -37,7 +37,8 @@ public:
 	virtual void ProcessMesh(
 		const FNamedDataMap& DatasIn,
 		const FDynamicMesh3& MeshIn,
-		FDynamicMesh3& MeshOut) override
+		FDynamicMesh3& MeshOut,
+		TUniquePtr<FEvaluationInfo>& EvaluationInfo) override
 	{
 		MeshOut = MeshIn;
 		TransferMaterialIDs(DatasIn, &MeshOut);
@@ -45,7 +46,8 @@ public:
 
 	virtual void ProcessMeshInPlace(
 		const FNamedDataMap& DatasIn,
-		FDynamicMesh3& MeshInOut) override
+		FDynamicMesh3& MeshInOut,
+		TUniquePtr<FEvaluationInfo>& EvaluationInfo) override
 	{
 		TransferMaterialIDs(DatasIn, &MeshInOut);
 	}
diff --git a/Engine/Plugins/Experimental/GeometryFlow/Source/GeometryFlowMeshProcessingEditor/Private/MeshProcessingNodes/MeshAutoGenerateUVsNode.cpp b/Engine/Plugins/Experimental/GeometryFlow/Source/GeometryFlowMeshProcessingEditor/Private/MeshProcessingNodes/MeshAutoGenerateUVsNode.cpp
--- a/Engine/Plugins/Experimental/GeometryFlow/Source/GeometryFlowMeshProcessingEditor/Private/MeshProcessingNodes/MeshAutoGenerateUVsNode.cpp
+++ b/Engine/Plugins/Experimental/GeometryFlow/Source/GeometryFlowMeshProcessingEditor/Private/MeshProcessingNodes/MeshAutoGenerateUVsNode.cpp
@@ -8,7 +8,10 @@
 using namespace UE::GeometryFlow;
 
 
-void FMeshAutoGenerateUVsNode::GenerateUVs(const FDynamicMesh3& MeshIn, const FMeshAutoGenerateUVsSettings& Settings, FDynamicMesh3& MeshOut)
+void FMeshAutoGenerateUVsNode::GenerateUVs(const FDynamicMesh3& MeshIn, 
+	const FMeshAutoGenerateUVsSettings& Settings, 
+	FDynamicMesh3& MeshOut,
+	TUniquePtr<FEvaluationInfo>& EvaluationInfo)
 {
 	// this is horrible - have to copy input mesh so that we can pass a TSharedPtr
 	TSharedPtr<FDynamicMesh3, ESPMode::ThreadSafe> InputMesh = MakeShared<FDynamicMesh3, ESPMode::ThreadSafe>(MeshIn);
@@ -42,9 +45,13 @@ void FMeshAutoGenerateUVsNode::GenerateUVs(const FDynamicMesh3& MeshIn, const FM
 		ParameterizeMeshOp.Method = UE::Geometry::EParamOpBackend::PatchBuilder;
 	}
 
-	FProgressCancel Progress;
-	ParameterizeMeshOp.CalculateResult(&Progress);
-	TUniquePtr<FDynamicMesh3> ResultMesh = ParameterizeMeshOp.ExtractResult();
+	ParameterizeMeshOp.CalculateResult(EvaluationInfo->Progress);
 
+	if (EvaluationInfo && EvaluationInfo->Progress && EvaluationInfo->Progress->Cancelled())
+	{
+		return;
+	}
+
+	TUniquePtr<FDynamicMesh3> ResultMesh = ParameterizeMeshOp.ExtractResult();
 	MeshOut = MoveTemp(*ResultMesh);
-}
\ No newline at end of file
+}
diff --git a/Engine/Plugins/Experimental/GeometryFlow/Source/GeometryFlowMeshProcessingEditor/Public/MeshProcessingNodes/MeshAutoGenerateUVsNode.h b/Engine/Plugins/Experimental/GeometryFlow/Source/GeometryFlowMeshProcessingEditor/Public/MeshProcessingNodes/MeshAutoGenerateUVsNode.h
--- a/Engine/Plugins/Experimental/GeometryFlow/Source/GeometryFlowMeshProcessingEditor/Public/MeshProcessingNodes/MeshAutoGenerateUVsNode.h
+++ b/Engine/Plugins/Experimental/GeometryFlow/Source/GeometryFlowMeshProcessingEditor/Public/MeshProcessingNodes/MeshAutoGenerateUVsNode.h
@@ -59,12 +59,16 @@ public:
 		const FNamedDataMap& DatasIn,
 		const FMeshAutoGenerateUVsSettings& Settings,
 		const FDynamicMesh3& MeshIn,
-		FDynamicMesh3& MeshOut) override
+		FDynamicMesh3& MeshOut,
+		TUniquePtr<FEvaluationInfo>& EvaluationInfo) override
 	{
-		GenerateUVs(MeshIn, Settings, MeshOut);
+		GenerateUVs(MeshIn, Settings, MeshOut, EvaluationInfo);
 	}
 
-	virtual void GenerateUVs(const FDynamicMesh3& MeshIn, const FMeshAutoGenerateUVsSettings& Settings, FDynamicMesh3& MeshOut);
+	virtual void GenerateUVs(const FDynamicMesh3& MeshIn, 
+		const FMeshAutoGenerateUVsSettings& Settings, 
+		FDynamicMesh3& MeshOut, 
+		TUniquePtr<FEvaluationInfo>& EvaluationInfo);
 
 };
 
