Fix a crash when shutting down editor after cancelling AutoLOD tool:
- Added a lock in EvaluateResult -- the FGraph destructor can't run until EvaluateResult is finished
- Hold on to shared node locks locally, just in case the FGraph object gets deleted while a node is running

#rnx

diff --git a/Engine/Plugins/Experimental/GeometryFlow/Source/GeometryFlowCore/Private/GeometryFlowGraph.cpp b/Engine/Plugins/Experimental/GeometryFlow/Source/GeometryFlowCore/Private/GeometryFlowGraph.cpp
--- a/Engine/Plugins/Experimental/GeometryFlow/Source/GeometryFlowCore/Private/GeometryFlowGraph.cpp
+++ b/Engine/Plugins/Experimental/GeometryFlow/Source/GeometryFlowCore/Private/GeometryFlowGraph.cpp
@@ -16,6 +16,11 @@ namespace
 	// Can't seem to use threadpool in GeometryProcessingUnitTests. GThreadPool is null.
 }
 
+FGraph::~FGraph()
+{
+	EvaluateLock.Lock();    // Wait for any EvaluateResult to finish and also prevent any new evals from happening while destructing
+}
+
 TSafeSharedPtr<FNode> FGraph::FindNode(FHandle NodeHandle) const
 {
 	const FNodeInfo* Found = AllNodes.Find(NodeHandle);
@@ -311,15 +316,19 @@ TSafeSharedPtr<IData> FGraph::ComputeOutputData(
 	FNamedDataMap DataOut;
 	DataOut.Add(OutputName);
 	
-	FindNodeLock(NodeHandle)->WriteLock();
+	TSafeSharedPtr<FRWLock> NodeLock = FindNodeLock(NodeHandle);
+	NodeLock->WriteLock();
 	Node->Evaluate(DataIn, DataOut, EvaluationInfo);
-	FindNodeLock(NodeHandle)->WriteUnlock();
+	NodeLock->WriteUnlock();
 
-	EvaluationInfo->CountEvaluation(Node.Get());
-
-	if (EvaluationInfo && EvaluationInfo->Progress && EvaluationInfo->Progress->Cancelled())
+	if (EvaluationInfo)
 	{
-		return nullptr;
+		EvaluationInfo->CountEvaluation(Node.Get());
+
+		if (EvaluationInfo->Progress && EvaluationInfo->Progress->Cancelled())
+		{
+			return nullptr;
+		}
 	}
 
 	// collect (and optionally take/steal) desired Output data
diff --git a/Engine/Plugins/Experimental/GeometryFlow/Source/GeometryFlowCore/Public/GeometryFlowGraph.h b/Engine/Plugins/Experimental/GeometryFlow/Source/GeometryFlowCore/Public/GeometryFlowGraph.h
--- a/Engine/Plugins/Experimental/GeometryFlow/Source/GeometryFlowCore/Public/GeometryFlowGraph.h
+++ b/Engine/Plugins/Experimental/GeometryFlow/Source/GeometryFlowCore/Public/GeometryFlowGraph.h
@@ -39,6 +39,8 @@ class GEOMETRYFLOWCORE_API FGraph
 {
 public:
 
+	~FGraph();
+
 	struct FHandle
 	{
 		static const int32 InvalidHandle = -1;
@@ -80,22 +82,50 @@ public:
 	TSet<FHandle> GetNodesWithNoConnectedInputs() const;
 
 	template<typename T>
-	EGeometryFlowResult EvaluateResult(FHandle Node, 
-									   FString OutputName, 
-									   T& Storage, 
-									   int32 StorageTypeIdentifier,
-									   FProgressCancel* Progress)
+	EGeometryFlowResult EvaluateResult(
+		FHandle Node,
+		FString OutputName,
+		T& Storage,
+		int32 StorageTypeIdentifier,
+		TUniquePtr<FEvaluationInfo>& EvaluationInfo,
+		bool bTryTakeResult)
 	{
-		TUniquePtr<FEvaluationInfo> EvalInfo = MakeUnique<FEvaluationInfo>();
-		EvalInfo->Progress = Progress;
-		return EvaluateResult(Node, OutputName, Storage, StorageTypeIdentifier, EvalInfo);
+		EvaluateLock.Lock();
+		EGeometryFlowResult Result;
+		Result = EvaluateResultInternal(Node, OutputName, Storage, StorageTypeIdentifier, EvaluationInfo, bTryTakeResult);
+		EvaluateLock.Unlock();
+		return Result;
 	}
 
+	template<typename NodeType>
+	EGeometryFlowResult ApplyToNodeOfType(
+		FHandle NodeHandle, 
+		TFunctionRef<void(NodeType&)> ApplyFunc)
+	{
+		TSafeSharedPtr<FNode> FoundNode = FindNode(NodeHandle);
+		if (FoundNode)
+		{
+			NodeType* Value = static_cast<NodeType*>(FoundNode.Get());
+			ApplyFunc(*Value);
+			return EGeometryFlowResult::Ok;
+		}
+		return EGeometryFlowResult::NodeDoesNotExist;
+	}
+
+
+
+	void ConfigureCachingStrategy(ENodeCachingStrategy NewStrategy);
+	EGeometryFlowResult SetNodeCachingStrategy(FHandle NodeHandle, ENodeCachingStrategy Strategy);
+
+	FString DebugDumpGraph(TFunction<bool(TSafeSharedPtr<FNode>)> IncludeNodeFn) const;
+
+protected:
+
 	template<typename T>
-	EGeometryFlowResult EvaluateResult(
-		FHandle Node, 
-		FString OutputName, 
-		T& Storage, 
+	EGeometryFlowResult EvaluateResultInternal(
+		FHandle Node,
+		FString OutputName,
+		T& Storage,
 		int32 StorageTypeIdentifier,
 		TUniquePtr<FEvaluationInfo>& EvaluationInfo,
 		bool bTryTakeResult)
@@ -132,30 +162,6 @@ public:
 	}
 
 
-	template<typename NodeType>
-	EGeometryFlowResult ApplyToNodeOfType(
-		FHandle NodeHandle, 
-		TFunctionRef<void(NodeType&)> ApplyFunc)
-	{
-		TSafeSharedPtr<FNode> FoundNode = FindNode(NodeHandle);
-		if (FoundNode)
-		{
-			NodeType* Value = static_cast<NodeType*>(FoundNode.Get());
-			ApplyFunc(*Value);
-			return EGeometryFlowResult::Ok;
-		}
-		return EGeometryFlowResult::NodeDoesNotExist;
-	}
-
-
-
-	void ConfigureCachingStrategy(ENodeCachingStrategy NewStrategy);
-	EGeometryFlowResult SetNodeCachingStrategy(FHandle NodeHandle, ENodeCachingStrategy Strategy);
-
-	FString DebugDumpGraph(TFunction<bool(TSafeSharedPtr<FNode>)> IncludeNodeFn) const;
-
-protected:
-
 	friend class FGeometryFlowExecutor;
 
 	int32 NodeCounter = 0;
@@ -188,6 +194,8 @@ protected:
 		FString OutputName, 
 		TUniquePtr<FEvaluationInfo>& EvaluationInfo,
 		bool bStealOutputData = false);
+
+	FCriticalSection EvaluateLock;
 };
 
 
