AutoLOD: the FProgressCancel object wasn't being passed to mesh simplification functions

#rnx

diff --git a/Engine/Plugins/Experimental/GeometryFlow/Source/GeometryFlowMeshProcessing/Public/MeshProcessingNodes/MeshSimplifyNode.h b/Engine/Plugins/Experimental/GeometryFlow/Source/GeometryFlowMeshProcessing/Public/MeshProcessingNodes/MeshSimplifyNode.h
--- a/Engine/Plugins/Experimental/GeometryFlow/Source/GeometryFlowMeshProcessing/Public/MeshProcessingNodes/MeshSimplifyNode.h
+++ b/Engine/Plugins/Experimental/GeometryFlow/Source/GeometryFlowMeshProcessing/Public/MeshProcessingNodes/MeshSimplifyNode.h
@@ -74,7 +74,7 @@ public:
 		TUniquePtr<FEvaluationInfo>& EvaluationInfo) override
 	{
 		MeshOut.Copy(MeshIn, true, true, true, !Settings.bDiscardAttributes);
-		ApplySimplify(Settings, MeshOut);
+		ApplySimplify(Settings, MeshOut, EvaluationInfo);
 	}
 
 	virtual void ProcessMeshInPlace(
@@ -83,15 +83,20 @@ public:
 		FDynamicMesh3& MeshInOut,
 		TUniquePtr<FEvaluationInfo>& EvaluationInfo)
 	{
-		ApplySimplify(Settings, MeshInOut);
+		ApplySimplify(Settings, MeshInOut, EvaluationInfo);
 	}
 
 
 	template<typename SimplifierType> 
-	void DoSimplifyOfType(const FMeshSimplifySettings& Settings, FDynamicMesh3* TargetMesh)
+	void DoSimplifyOfType(const FMeshSimplifySettings& Settings, FDynamicMesh3* TargetMesh, TUniquePtr<FEvaluationInfo>& EvaluationInfo)
 	{
 		SimplifierType Simplifier(TargetMesh);
 
+		if (EvaluationInfo && EvaluationInfo->Progress)
+		{
+			Simplifier.Progress = EvaluationInfo->Progress;
+		}
+
 		Simplifier.ProjectionMode = SimplifierType::ETargetProjectionMode::NoProjection;
 		Simplifier.DEBUG_CHECK_LEVEL = 0;
 
@@ -148,7 +153,7 @@ public:
 	}
 
 
-	void ApplySimplify(const FMeshSimplifySettings& Settings, FDynamicMesh3& MeshInOut)
+	void ApplySimplify(const FMeshSimplifySettings& Settings, FDynamicMesh3& MeshInOut, TUniquePtr<FEvaluationInfo>& EvaluationInfo)
 	{
 		if (Settings.bDiscardAttributes)
 		{
@@ -157,15 +162,15 @@ public:
 
 		if (Settings.SimplifyType == EMeshSimplifyType::Standard)
 		{
-			DoSimplifyOfType<FQEMSimplification>(Settings, &MeshInOut);
+			DoSimplifyOfType<FQEMSimplification>(Settings, &MeshInOut, EvaluationInfo);
 		}
 		else if (Settings.SimplifyType == EMeshSimplifyType::VolumePreserving)
 		{
-			DoSimplifyOfType<FVolPresMeshSimplification>(Settings, &MeshInOut);
+			DoSimplifyOfType<FVolPresMeshSimplification>(Settings, &MeshInOut, EvaluationInfo);
 		}
 		else if (Settings.SimplifyType == EMeshSimplifyType::AttributeAware)
 		{
-			DoSimplifyOfType<FAttrMeshSimplification>(Settings, &MeshInOut);
+			DoSimplifyOfType<FAttrMeshSimplification>(Settings, &MeshInOut, EvaluationInfo);
 		}
 		else
 		{
