AutoLOD: prevent a hard crash if a node input mesh doesn't have attributes

#rnx

diff --git a/Engine/Plugins/Experimental/GeometryFlow/Source/GeometryFlowMeshProcessing/Private/MeshProcessingNodes/MeshRepackUVsNode.cpp b/Engine/Plugins/Experimental/GeometryFlow/Source/GeometryFlowMeshProcessing/Private/MeshProcessingNodes/MeshRepackUVsNode.cpp
--- a/Engine/Plugins/Experimental/GeometryFlow/Source/GeometryFlowMeshProcessing/Private/MeshProcessingNodes/MeshRepackUVsNode.cpp
+++ b/Engine/Plugins/Experimental/GeometryFlow/Source/GeometryFlowMeshProcessing/Private/MeshProcessingNodes/MeshRepackUVsNode.cpp
@@ -10,6 +10,11 @@ using namespace UE::GeometryFlow;
 
 void FMeshRepackUVsNode::RepackUVsForMesh(FDynamicMesh3& EditMesh, const FMeshRepackUVsSettings& Settings)
 {
+	if ( !(EditMesh.HasAttributes() && EditMesh.Attributes()->GetUVLayer(Settings.UVLayer)) )
+	{
+		return;
+	}
+
 	FDynamicMeshUVOverlay* UVLayer = EditMesh.Attributes()->GetUVLayer(Settings.UVLayer);
 
 	UVLayer->SplitBowties();
