AutoLOD: don't run UV packing if the input mesh is empty

#rnx

diff --git a/Engine/Plugins/Experimental/GeometryFlow/Source/GeometryFlowMeshProcessing/Private/MeshProcessingNodes/MeshRepackUVsNode.cpp b/Engine/Plugins/Experimental/GeometryFlow/Source/GeometryFlowMeshProcessing/Private/MeshProcessingNodes/MeshRepackUVsNode.cpp
--- a/Engine/Plugins/Experimental/GeometryFlow/Source/GeometryFlowMeshProcessing/Private/MeshProcessingNodes/MeshRepackUVsNode.cpp
+++ b/Engine/Plugins/Experimental/GeometryFlow/Source/GeometryFlowMeshProcessing/Private/MeshProcessingNodes/MeshRepackUVsNode.cpp
@@ -15,6 +15,11 @@ void FMeshRepackUVsNode::RepackUVsForMesh(FDynamicMesh3& EditMesh, const FMeshRe
 		return;
 	}
 
+	if (EditMesh.TriangleCount() == 0)
+	{
+		return;
+	}
+
 	FDynamicMeshUVOverlay* UVLayer = EditMesh.Attributes()->GetUVLayer(Settings.UVLayer);
 
 	UVLayer->SplitBowties();
