AutoLOD: if node computation is cancelled, do not update the cached outputs

#rnx

diff --git a/Engine/Plugins/Experimental/GeometryFlow/Source/GeometryFlowMeshProcessing/Public/MeshProcessingNodes/MeshProcessingBaseNodes.h b/Engine/Plugins/Experimental/GeometryFlow/Source/GeometryFlowMeshProcessing/Public/MeshProcessingNodes/MeshProcessingBaseNodes.h
--- a/Engine/Plugins/Experimental/GeometryFlow/Source/GeometryFlowMeshProcessing/Public/MeshProcessingNodes/MeshProcessingBaseNodes.h
+++ b/Engine/Plugins/Experimental/GeometryFlow/Source/GeometryFlowMeshProcessing/Public/MeshProcessingNodes/MeshProcessingBaseNodes.h
@@ -6,7 +6,7 @@
 #include "GeometryFlowCoreNodes.h"
 #include "GeometryFlowNodeUtil.h"
 #include "DataTypes/DynamicMeshData.h"
-
+#include "Util/ProgressCancel.h"
 
 namespace UE
 {
@@ -77,6 +77,9 @@ public:
 			{
 				if (bRecomputeRequired)
 				{
+					// if execution is interrupted by ProgressCancel, make sure this node is recomputed next time through
+					ClearAllOutputs();
+
 					bool bIsMeshMutable = DatasIn.GetDataFlags(InParamMesh()).bIsMutableData;
 					if (bIsMeshMutable)
 					{
@@ -86,6 +89,11 @@ public:
 						MeshArg->GiveTo<FDynamicMesh3>(EditableMesh, (int)EMeshProcessingDataTypes::DynamicMesh);
 						ProcessMeshInPlace(DatasIn, EditableMesh, EvaluationInfo);
 
+						if (EvaluationInfo && EvaluationInfo->Progress && EvaluationInfo->Progress->Cancelled())
+						{
+							return;
+						}
+
 						// store new result
 						TSafeSharedPtr<FDataDynamicMesh> Result = MakeSafeShared<FDataDynamicMesh>();
 						Result->MoveData(MoveTemp(EditableMesh));
@@ -102,6 +110,11 @@ public:
 						FDynamicMesh3 ResultMesh;
 						ProcessMesh(DatasIn, SourceMesh, ResultMesh, EvaluationInfo);
 
+						if (EvaluationInfo && EvaluationInfo->Progress && EvaluationInfo->Progress->Cancelled())
+						{
+							return;
+						}
+
 						// store new result
 						TSafeSharedPtr<FDataDynamicMesh> Result = MakeSafeShared<FDataDynamicMesh>();
 						Result->MoveData(MoveTemp(ResultMesh));
@@ -218,6 +231,9 @@ public:
 			{
 				if (bRecomputeRequired)
 				{
+					// if execution is interrupted by ProgressCancel, make sure this node is recomputed next time through
+					ClearAllOutputs();
+
 					// always make a copy of settings
 					SettingsType Settings;
 					SettingsArg->GetDataCopy(Settings, SettingsType::DataTypeIdentifier);
@@ -231,6 +247,11 @@ public:
 						MeshArg->GiveTo<FDynamicMesh3>(EditableMesh, (int)EMeshProcessingDataTypes::DynamicMesh);
 						ProcessMeshInPlace(DatasIn, Settings, EditableMesh, EvaluationInfo);
 
+						if (EvaluationInfo && EvaluationInfo->Progress && EvaluationInfo->Progress->Cancelled())
+						{
+							return;
+						}
+
 						// store new result
 						TSafeSharedPtr<FDataDynamicMesh> Result = MakeSafeShared<FDataDynamicMesh>();
 						Result->MoveData(MoveTemp(EditableMesh));
@@ -247,6 +268,11 @@ public:
 						FDynamicMesh3 ResultMesh;
 						ProcessMesh(DatasIn, Settings, SourceMesh, ResultMesh, EvaluationInfo);
 
+						if (EvaluationInfo && EvaluationInfo->Progress && EvaluationInfo->Progress->Cancelled())
+						{
+							return;
+						}
+
 						// store new result
 						TSafeSharedPtr<FDataDynamicMesh> Result = MakeSafeShared<FDataDynamicMesh>();
 						Result->MoveData(MoveTemp(ResultMesh));
