Fixed errors when compiling for Windows with Clang

#rnx

diff --git a/Engine/Plugins/Experimental/GeometryFlow/Source/GeometryFlowMeshProcessing/Public/MeshProcessingNodes/MeshProcessingBaseNodes.h b/Engine/Plugins/Experimental/GeometryFlow/Source/GeometryFlowMeshProcessing/Public/MeshProcessingNodes/MeshProcessingBaseNodes.h
--- a/Engine/Plugins/Experimental/GeometryFlow/Source/GeometryFlowMeshProcessing/Public/MeshProcessingNodes/MeshProcessingBaseNodes.h
+++ b/Engine/Plugins/Experimental/GeometryFlow/Source/GeometryFlowMeshProcessing/Public/MeshProcessingNodes/MeshProcessingBaseNodes.h
@@ -168,6 +168,19 @@ public:
 
 
 
+namespace Private
+{
+
+// This function is a workaround for an issue with Clang targeting Windows.
+// A function in a class template cannot contain log statements if that class template
+// is derived from by a type annotated with {Module}_API.
+inline void LogMeshProcessing(const TCHAR* Identifier, const TCHAR* Message)
+{
+	UE_LOG(LogGeometryFlowMeshProcessing, Display, TEXT("[%s]  %s"), Identifier, Message);
+}
+
+} // end namespace Private
+
 template<typename SettingsType>
 class TProcessMeshWithSettingsBaseNode : public FNode
 {
@@ -241,7 +254,7 @@ public:
 					bool bIsMeshMutable = DatasIn.GetDataFlags(InParamMesh()).bIsMutableData;
 					if (bIsMeshMutable)
 					{
-						UE_LOG(LogGeometryFlowMeshProcessing, Display, TEXT("[%s]  RECOMPUTING MeshOp In Place!"), *GetIdentifier());
+						Private::LogMeshProcessing(*GetIdentifier(), TEXT("RECOMPUTING MeshOp In Place!"));
 
 						FDynamicMesh3 EditableMesh;
 						MeshArg->GiveTo<FDynamicMesh3>(EditableMesh, (int)EMeshProcessingDataTypes::DynamicMesh);
@@ -259,7 +272,7 @@ public:
 					}
 					else
 					{
-						UE_LOG(LogGeometryFlowMeshProcessing, Display, TEXT("[%s]  RECOMPUTING MeshOp"), *GetIdentifier());
+						Private::LogMeshProcessing(*GetIdentifier(), TEXT("RECOMPUTING MeshOp"));
 
 						// do we ever want to support using a copy of the source mesh?
 						const FDynamicMesh3& SourceMesh = MeshArg->GetDataConstRef<FDynamicMesh3>((int)EMeshProcessingDataTypes::DynamicMesh);
