GeometryFlow: Adding Serialization to the geometry flow graph  and its nodes, and some additional editing functionality to the graph.  Also changing the geometryflow settings nodes to use ustructs instead of just structs.

diff --git a/Engine/Plugins/Editor/MeshLODToolset/Source/MeshLODToolset/Private/Graphs/GenerateMeshLODGraph.cpp b/Engine/Plugins/Editor/MeshLODToolset/Source/MeshLODToolset/Private/Graphs/GenerateMeshLODGraph.cpp
--- a/Engine/Plugins/Editor/MeshLODToolset/Source/MeshLODToolset/Private/Graphs/GenerateMeshLODGraph.cpp
+++ b/Engine/Plugins/Editor/MeshLODToolset/Source/MeshLODToolset/Private/Graphs/GenerateMeshLODGraph.cpp
@@ -35,7 +35,8 @@ void FGenerateMeshLODGraph::SetSourceMesh(const FDynamicMesh3& SourceMeshIn)
 
 void FGenerateMeshLODGraph::UpdatePreFilterSettings(const FMeshLODGraphPreFilterSettings& PreFilterSettings)
 {
-	UpdateSourceNodeValue<FNameSourceNode>(*Graph, FilterGroupsLayerNameNode, PreFilterSettings.FilterGroupLayerName);
+	FName FilterGroupLayerName = { PreFilterSettings.FilterGroupLayerName };
+	UpdateSourceNodeValue<FNameSourceNode>(*Graph, FilterGroupsLayerNameNode, FilterGroupLayerName);
 	CurrentPreFilterSettings = PreFilterSettings;
 }
 
@@ -499,22 +500,22 @@ void FGenerateMeshLODGraph::BuildGraph(const FDynamicMesh3* SourceMeshHint)
 	// TODO: can use simpler settings if there is only one material...
 	FMeshSimplifySettings SimplifySettings;
 	SimplifySettings.bDiscardAttributes = false;
-	SimplifySettings.SimplifyType = EMeshSimplifyType::AttributeAware;
+	SimplifySettings.SimplifyType = ToUEnum(EMeshSimplifyType::AttributeAware);
 	//SimplifySettings.TargetType = EMeshSimplifyTargetType::TriangleCount;
 	//SimplifySettings.TargetCount = 500;
-	SimplifySettings.TargetType = EMeshSimplifyTargetType::GeometricDeviation;
+	SimplifySettings.TargetType = ToUEnum(EMeshSimplifyTargetType::GeometricDeviation);
 	SimplifySettings.GeometricTolerance = 0.5;
 	SimplifySettings.TargetCount = 500;
-	SimplifySettings.MaterialBorderConstraints = EEdgeRefineFlags::NoFlip;
+	SimplifySettings.MaterialBorderConstraints = ToUEnum(EEdgeRefineFlags::NoFlip);
 	UpdateSimplifySettings(SimplifySettings);
 
 	FMeshNormalsSettings NormalsSettings;
-	NormalsSettings.NormalsType = EComputeNormalsType::FromFaceAngleThreshold;
+	NormalsSettings.NormalsType = ToUEnum(EComputeNormalsType::FromFaceAngleThreshold);
 	NormalsSettings.AngleThresholdDeg = 60.0;
 	UpdateNormalsSettings(NormalsSettings);
 
 	FMeshAutoGenerateUVsSettings AutoUVSettings;
-	AutoUVSettings.Method = EAutoUVMethod::PatchBuilder;
+	AutoUVSettings.Method = ToUEnum(EAutoUVMethod::PatchBuilder);
 	AutoUVSettings.UVAtlasNumCharts = 20;
 	AutoUVSettings.UVAtlasStretch = 0.1;
 	AutoUVSettings.XAtlasMaxIterations = 1;
@@ -540,7 +541,7 @@ void FGenerateMeshLODGraph::BuildGraph(const FDynamicMesh3* SourceMeshHint)
 
 
 	FMeshMakeBakingCacheSettings BakeCacheSettings;
-	BakeCacheSettings.Dimensions = FImageDimensions(1024, 1024);
+	BakeCacheSettings.Dimensions = FIntPoint(1024, 1024);
 	BakeCacheSettings.Thickness = 5.0;
 	UpdateBakeCacheSettings(BakeCacheSettings);
 
@@ -555,7 +556,7 @@ void FGenerateMeshLODGraph::BuildGraph(const FDynamicMesh3* SourceMeshHint)
 	UpdateCollisionGroupLayerName(CollisionGroupLayerName);
 
 	FGenerateSimpleCollisionSettings GenSimpleCollisionSettings;
-	GenSimpleCollisionSettings.Type = UE::GeometryFlow::ESimpleCollisionGeometryType::AlignedBoxes;
+	GenSimpleCollisionSettings.Type = EGeometryFlow_SimpleCollisionGeometryType::AlignedBoxes;
 	UpdateGenerateSimpleCollisionSettings(GenSimpleCollisionSettings);
 
 	TArray<float> Weights;
diff --git a/Engine/Plugins/Editor/MeshLODToolset/Source/MeshLODToolset/Private/Graphs/GenerateStaticMeshLODProcess.cpp b/Engine/Plugins/Editor/MeshLODToolset/Source/MeshLODToolset/Private/Graphs/GenerateStaticMeshLODProcess.cpp
--- a/Engine/Plugins/Editor/MeshLODToolset/Source/MeshLODToolset/Private/Graphs/GenerateStaticMeshLODProcess.cpp
+++ b/Engine/Plugins/Editor/MeshLODToolset/Source/MeshLODToolset/Private/Graphs/GenerateStaticMeshLODProcess.cpp
@@ -617,7 +617,7 @@ bool UGenerateStaticMeshLODProcess::InitializeGenerator()
 		FGenerateStaticMeshLODProcess_NormalsSettings::MapMethodType((int32)Generator->GetCurrentNormalsSettings().NormalsType, false));
 	CurrentSettings_Normals.Angle = Generator->GetCurrentNormalsSettings().AngleThresholdDeg;
 
-	CurrentSettings_Texture.BakeResolution = (EGenerateStaticMeshLODBakeResolution)Generator->GetCurrentBakeCacheSettings().Dimensions.GetWidth();
+	CurrentSettings_Texture.BakeResolution = (EGenerateStaticMeshLODBakeResolution)Generator->GetCurrentBakeCacheSettings().Dimensions.X;
 	CurrentSettings_Texture.BakeThickness = Generator->GetCurrentBakeCacheSettings().Thickness;
 
 	CurrentSettings_UV.UVMethod = static_cast<EGenerateStaticMeshLODProcess_AutoUVMethod>(Generator->GetCurrentAutoUVSettings().Method);
@@ -636,7 +636,7 @@ bool UGenerateStaticMeshLODProcess::InitializeGenerator()
 	CurrentSettings_Collision.PrefilterGridResolution = SimpleCollisionSettings.ConvexHullSettings.PrefilterGridResolution;
 	CurrentSettings_Collision.bSimplifyPolygons = SimpleCollisionSettings.SweptHullSettings.bSimplifyPolygons;
 	CurrentSettings_Collision.HullTolerance = SimpleCollisionSettings.SweptHullSettings.HullTolerance;
-	FMeshSimpleShapeApproximation::EProjectedHullAxisMode RHSMode = SimpleCollisionSettings.SweptHullSettings.SweepAxis;
+	FMeshSimpleShapeApproximation::EProjectedHullAxisMode RHSMode = static_cast<FMeshSimpleShapeApproximation::EProjectedHullAxisMode>(SimpleCollisionSettings.SweptHullSettings.SweepAxis);
 	CurrentSettings_Collision.SweepAxis = static_cast<EGenerateStaticMeshLODProjectedHullAxisMode>(RHSMode);
 
 
@@ -730,7 +730,7 @@ void UGenerateStaticMeshLODProcess::UpdateSimplifySettings(const FGenerateStatic
 		NewCombinedSettings.Tolerance != CurrentSettings_Simplify.Tolerance )
 	{
 		UE::GeometryFlow::FMeshSimplifySettings NewSimplifySettings = Generator->GetCurrentSimplifySettings();
-		NewSimplifySettings.TargetType = static_cast<UE::GeometryFlow::EMeshSimplifyTargetType>(NewCombinedSettings.Method);
+		NewSimplifySettings.TargetType = ToUEnum(static_cast<UE::GeometryFlow::EMeshSimplifyTargetType>(NewCombinedSettings.Method));
 		NewSimplifySettings.TargetCount = NewCombinedSettings.TargetCount;
 		NewSimplifySettings.TargetFraction = NewCombinedSettings.TargetPercentage / 100.0f;
 		NewSimplifySettings.GeometricTolerance = NewCombinedSettings.Tolerance;
@@ -774,8 +774,8 @@ void UGenerateStaticMeshLODProcess::UpdateNormalsSettings(const FGenerateStaticM
 	{
 		UE::GeometryFlow::FMeshNormalsSettings NewNormalsSettings = Generator->GetCurrentNormalsSettings();
 
-		NewNormalsSettings.NormalsType = static_cast<UE::GeometryFlow::EComputeNormalsType>(
-			FGenerateStaticMeshLODProcess_NormalsSettings::MapMethodType((int32)NewCombinedSettings.Method, true));
+		NewNormalsSettings.NormalsType = ToUEnum(static_cast<UE::GeometryFlow::EComputeNormalsType>(
+			FGenerateStaticMeshLODProcess_NormalsSettings::MapMethodType((int32)NewCombinedSettings.Method, true)));
 		NewNormalsSettings.AngleThresholdDeg = NewCombinedSettings.Angle;
 
 		Generator->UpdateNormalsSettings(NewNormalsSettings);
@@ -793,7 +793,7 @@ void UGenerateStaticMeshLODProcess::UpdateTextureSettings(const FGenerateStaticM
 		 (NewCombinedSettings.bCombineTextures != CurrentSettings_Texture.bCombineTextures))
 	{
 		UE::GeometryFlow::FMeshMakeBakingCacheSettings NewBakeSettings = Generator->GetCurrentBakeCacheSettings();
-		NewBakeSettings.Dimensions = FImageDimensions((int32)NewCombinedSettings.BakeResolution, (int32)NewCombinedSettings.BakeResolution);
+		NewBakeSettings.Dimensions = FIntPoint((int32)NewCombinedSettings.BakeResolution, (int32)NewCombinedSettings.BakeResolution);
 		NewBakeSettings.Thickness = NewCombinedSettings.BakeThickness;
 		Generator->UpdateBakeCacheSettings(NewBakeSettings);
 	}
@@ -816,7 +816,7 @@ void UGenerateStaticMeshLODProcess::UpdateUVSettings(const FGenerateStaticMeshLO
 		|| NewCombinedSettings.PatchBuilder.SmoothingAlpha != CurrentSettings_UV.PatchBuilder.SmoothingAlpha  ) 
 	{
 		UE::GeometryFlow::FMeshAutoGenerateUVsSettings NewAutoUVSettings = Generator->GetCurrentAutoUVSettings();
-		NewAutoUVSettings.Method = (EAutoUVMethod)(int)NewCombinedSettings.UVMethod;
+		NewAutoUVSettings.Method = ToUEnum( (EAutoUVMethod)(int)NewCombinedSettings.UVMethod);
 		//NewAutoUVSettings.UVAtlasStretch = 0.5;
 		NewAutoUVSettings.UVAtlasNumCharts = NewCombinedSettings.NumUVAtlasCharts;
 		//NewAutoUVSettings.XAtlasMaxIterations = 1;
@@ -853,13 +853,13 @@ void UGenerateStaticMeshLODProcess::UpdateCollisionSettings(const FGenerateStati
 		NewCombinedSettings.CollisionType != CurrentSettings_Collision.CollisionType)
 	{
 		UE::GeometryFlow::FGenerateSimpleCollisionSettings NewGenCollisionSettings = Generator->GetCurrentGenerateSimpleCollisionSettings();
-		NewGenCollisionSettings.Type = static_cast<UE::GeometryFlow::ESimpleCollisionGeometryType>(NewCombinedSettings.CollisionType);
+		NewGenCollisionSettings.Type = static_cast<EGeometryFlow_SimpleCollisionGeometryType>(NewCombinedSettings.CollisionType);
 		NewGenCollisionSettings.ConvexHullSettings.SimplifyToTriangleCount = NewCombinedSettings.ConvexTriangleCount;
 		NewGenCollisionSettings.ConvexHullSettings.bPrefilterVertices = NewCombinedSettings.bPrefilterVertices;
 		NewGenCollisionSettings.ConvexHullSettings.PrefilterGridResolution = NewCombinedSettings.PrefilterGridResolution;
 		NewGenCollisionSettings.SweptHullSettings.bSimplifyPolygons = NewCombinedSettings.bSimplifyPolygons;
 		NewGenCollisionSettings.SweptHullSettings.HullTolerance = NewCombinedSettings.HullTolerance;
-		NewGenCollisionSettings.SweptHullSettings.SweepAxis = static_cast<FMeshSimpleShapeApproximation::EProjectedHullAxisMode>(NewCombinedSettings.SweepAxis);
+		NewGenCollisionSettings.SweptHullSettings.SweepAxis = static_cast<EGeometryFlow_ProjectedHullAxisMode>(NewCombinedSettings.SweepAxis);
 		Generator->UpdateGenerateSimpleCollisionSettings(NewGenCollisionSettings);
 	}
 
diff --git a/Engine/Plugins/Editor/MeshLODToolset/Source/MeshLODToolset/Public/Graphs/GenerateMeshLODGraph.h b/Engine/Plugins/Editor/MeshLODToolset/Source/MeshLODToolset/Public/Graphs/GenerateMeshLODGraph.h
--- a/Engine/Plugins/Editor/MeshLODToolset/Source/MeshLODToolset/Public/Graphs/GenerateMeshLODGraph.h
+++ b/Engine/Plugins/Editor/MeshLODToolset/Source/MeshLODToolset/Public/Graphs/GenerateMeshLODGraph.h
@@ -57,7 +57,7 @@ public:
 	void UpdateCoreMeshGeneratorMode(ECoreMeshGeneratorMode NewMode);
 	ECoreMeshGeneratorMode GetCurrentCoreMeshGeneratorMode() const { return CurrentCoreMeshGeneratorMode; }
 
-	void UpdateSolidifySettings(const UE::GeometryFlow::FMeshSolidifySettings& SolidifySettings);
+	void UpdateSolidifySettings(const FMeshSolidifySettings& SolidifySettings);
 	const UE::GeometryFlow::FMeshSolidifySettings& GetCurrentSolidifySettings() const { return CurrentSolidifySettings; }
 
 	void UpdateMorphologySettings(const UE::GeometryFlow::FVoxClosureSettings& MorphologySettings);
diff --git a/Engine/Plugins/Experimental/GeometryFlow/Source/GeometryFlowCore/GeometryFlowCore.Build.cs b/Engine/Plugins/Experimental/GeometryFlow/Source/GeometryFlowCore/GeometryFlowCore.Build.cs
--- a/Engine/Plugins/Experimental/GeometryFlow/Source/GeometryFlowCore/GeometryFlowCore.Build.cs
+++ b/Engine/Plugins/Experimental/GeometryFlow/Source/GeometryFlowCore/GeometryFlowCore.Build.cs
@@ -18,7 +18,7 @@ public class GeometryFlowCore : ModuleRules
 		PrivateDependencyModuleNames.AddRange(
 			new string[]
 			{
-				//"CoreUObject",
+				"CoreUObject",
 				//"Engine"
 			}
 			);
diff --git a/Engine/Plugins/Experimental/GeometryFlow/Source/GeometryFlowCore/Private/GeometryFlowCoreModule.cpp b/Engine/Plugins/Experimental/GeometryFlow/Source/GeometryFlowCore/Private/GeometryFlowCoreModule.cpp
--- a/Engine/Plugins/Experimental/GeometryFlow/Source/GeometryFlowCore/Private/GeometryFlowCoreModule.cpp
+++ b/Engine/Plugins/Experimental/GeometryFlow/Source/GeometryFlowCore/Private/GeometryFlowCoreModule.cpp
@@ -3,11 +3,16 @@
 #include "GeometryFlowCoreModule.h"
 #include "Modules/ModuleManager.h"
 
+#include "GeometryFlowNodeFactory.h"
+
 #define LOCTEXT_NAMESPACE "FGeometryFlowCoreModule"
 
 void FGeometryFlowCoreModule::StartupModule()
 {
 	// This code will execute after your module is loaded into memory; the exact timing is specified in the .uplugin file per-module
+
+		// Register the nodes defined by this module with the GeometryFlowNodeFactory.
+	UE::GeometryFlow::RegisterCoreNodes();
 }
 
 void FGeometryFlowCoreModule::ShutdownModule()
diff --git a/Engine/Plugins/Experimental/GeometryFlow/Source/GeometryFlowCore/Private/GeometryFlowGraph.cpp b/Engine/Plugins/Experimental/GeometryFlow/Source/GeometryFlowCore/Private/GeometryFlowGraph.cpp
--- a/Engine/Plugins/Experimental/GeometryFlow/Source/GeometryFlowCore/Private/GeometryFlowGraph.cpp
+++ b/Engine/Plugins/Experimental/GeometryFlow/Source/GeometryFlowCore/Private/GeometryFlowGraph.cpp
@@ -3,6 +3,7 @@
 #include "GeometryFlowGraph.h"
 #include "Async/Async.h"
 #include "Algo/Find.h"
+#include "UObject/UE5MainStreamObjectVersion.h"
 
 using namespace UE::GeometryFlow;
 
@@ -75,7 +76,24 @@ ENodeCachingStrategy FGraph::GetCachingStrategyForNode(FHandle NodeHandle) const
 }
 
 
-EGeometryFlowResult FGraph::AddConnection(FHandle FromNode, FString FromOutput, FHandle ToNode, FString ToInput)
+bool  FGraph::IsDependent(FHandle DependentNode, FHandle IndependentNode) const
+{
+	const FNodeInfo* DependentNodeInfo = AllNodes.Find(DependentNode);
+	const FNodeInfo* IndependentNodeInfo = AllNodes.Find(IndependentNode);
+
+	if (!DependentNodeInfo || !IndependentNodeInfo) return false;
+
+	// visit all nodes up-stream from the DependentNode and return true if we find the IndependentNode
+	auto Visitor = [IndependentNode](FHandle UpStreamNode) { return IndependentNode == UpStreamNode; };
+	if (DependentNode == IndependentNode || VistDependencies(DependentNode, Visitor) == true)
+	{
+		return true;
+	}
+	return false;
+}
+
+
+EGeometryFlowResult FGraph::CanAddConnection(FHandle FromNode, FString FromOutput, FHandle ToNode, FString ToInput) const
 {
 	int32 FromType = -1, ToType = -2;
 	EGeometryFlowResult FromTypeResult = GetOutputTypeForNode(FromNode, FromOutput, FromType);
@@ -94,7 +112,7 @@ EGeometryFlowResult FGraph::AddConnection(FHandle FromNode, FString FromOutput,
 		return EGeometryFlowResult::UnmatchedTypes;
 	}
 
-	FConnection* Found = Algo::FindByPredicate(Connections, [&](const FConnection& Conn)
+	const FConnection* Found = Algo::FindByPredicate(Connections, [&](const FConnection& Conn)
 	{
 		return ((Conn.ToNode == ToNode) && (Conn.ToInput == ToInput));
 	});
@@ -103,12 +121,28 @@ EGeometryFlowResult FGraph::AddConnection(FHandle FromNode, FString FromOutput,
 		return EGeometryFlowResult::InputAlreadyConnected;
 	}
 
-	Connections.Add({ FromNode, FromOutput, ToNode, ToInput });
+	// is the 'FromNode' already dependent on the 'ToNode'?  If so, the connection would create a cycle.
+	if (IsDependent(FromNode, ToNode))
+	{
+		return EGeometryFlowResult::ConnectionRejectedCreatesCycle;
+	};
 
 	return EGeometryFlowResult::Ok;
 }
 
 
+EGeometryFlowResult FGraph::AddConnection(FHandle FromNode, FString FromOutput, FHandle ToNode, FString ToInput)
+{
+	EGeometryFlowResult Result = CanAddConnection(FromNode, FromOutput, ToNode, ToInput);
+
+	if (Result == EGeometryFlowResult::Ok)
+	{
+		Connections.Add({ FromNode, FromOutput, ToNode, ToInput });
+	}
+
+	return Result;
+}
+
 
 EGeometryFlowResult FGraph::InferConnection(FHandle FromNodeHandle, FHandle ToNodeHandle)
 {
@@ -120,6 +154,11 @@ EGeometryFlowResult FGraph::InferConnection(FHandle FromNodeHandle, FHandle ToNo
 		return EGeometryFlowResult::NodeDoesNotExist;
 	}
 
+	if (FromNodeHandle == ToNodeHandle)
+	{
+		return EGeometryFlowResult::ConnectionRejectedCreatesCycle;
+	}
+
 	FString FromOutputName;
 	FString ToInputName;
 	int32 TotalMatchesFound = 0;
@@ -198,6 +237,35 @@ TSet<FGraph::FHandle> FGraph::GetNodesWithNoConnectedInputs() const
 
 
 
+EGeometryFlowResult FGraph::RemoveConnectionForInput(FHandle ToNode, FString ToInput)
+{
+	for (int i = 0; i < Connections.Num(); ++i)
+	{
+		FConnection& Connection = Connections[i];
+		if (Connection.ToNode == ToNode && Connection.ToInput == ToInput)
+		{
+			Connections.RemoveAtSwap(i);
+			return EGeometryFlowResult::Ok;
+		}
+	}
+
+	return EGeometryFlowResult::ConnectionDoesNotExist;
+}
+
+EGeometryFlowResult FGraph::FindInputConnections(FHandle ToNode, TArray<FGraph::FConnection>& ConnectionsOut) const
+{
+	if (!AllNodes.Contains(ToNode)) return EGeometryFlowResult::NodeDoesNotExist;
+
+	for (const FConnection& Connection : Connections)
+	{
+		if (Connection.ToNode == ToNode)
+		{
+			ConnectionsOut.Add(Connection);
+		}
+	}
+	return EGeometryFlowResult::Ok;
+}
+
 
 EGeometryFlowResult FGraph::FindConnectionForInput(FHandle ToNode, FString ToInput, FConnection& ConnectionOut) const
 {
@@ -339,6 +407,112 @@ TSafeSharedPtr<IData> FGraph::ComputeOutputData(
 }
 
 
+bool FGraph::CanComputeOutput(FHandle NodeHandle, FString OutputName)
+{
+	TSafeSharedPtr<FNode> Node = FindNode(NodeHandle);
+	if (!Node) return false;
+
+	TArray<FEvalRequirement> InputRequirements;
+	Node->CollectRequirements({ OutputName }, InputRequirements);
+
+	const int NumInputRequirements = InputRequirements.Num();
+
+	TArray<TFuture<bool>> AsyncFutures;
+	AsyncFutures.Reserve(NumInputRequirements);
+
+	// This will recursively call CanComputeOutput() on those (Node/Output) pairs
+	for (int32 k = 0; k < NumInputRequirements; ++k)
+	{
+		const FString& InputName = InputRequirements[k].InputName;
+
+		// find the connection for this input
+		FConnection Connection;
+		EGeometryFlowResult FoundResult = FindConnectionForInput(NodeHandle, InputName, Connection);
+
+		if (FoundResult == EGeometryFlowResult::ConnectionDoesNotExist)
+		{
+			TSafeSharedPtr<IData> DefaultData = Node->GetDefaultInputData(InputName);
+			if (!DefaultData) // the node input is not connected and has no default value;
+			{
+				return false;
+			}
+			continue;
+		}
+
+
+		TFuture<bool> Future = Async(GeometryFlowGraphAsyncExecTarget,
+			[this, &InputName, Connection]()
+			{
+				// recursively look upstream from this Input via Connection
+				const bool ValidUpstream = CanComputeOutput(Connection.FromNode, Connection.FromOutput);
+				return ValidUpstream;
+
+			});
+
+		AsyncFutures.Emplace(MoveTemp(Future));
+	}
+
+	bool bAllInputsValid = true;
+	for (TFuture<bool>& Future : AsyncFutures)
+	{
+		Future.Wait();
+		bAllInputsValid = bAllInputsValid && Future.Get();
+	}
+
+	return bAllInputsValid;
+}
+
+bool FGraph::VistDependencies(FHandle NodeHandle, TFunctionRef<bool(FGraph::FHandle)> Visitor) const
+{
+	TSafeSharedPtr<FNode> Node = FindNode(NodeHandle);
+	if (!Node) return false;
+
+	TArray<FEvalRequirement> InputRequirements;
+	Node->CollectAllRequirements(InputRequirements);
+
+	const int NumInputRequirements = InputRequirements.Num();
+
+	TArray<TFuture<bool>> AsyncFutures;
+	AsyncFutures.Reserve(NumInputRequirements);
+
+	// This will recursively call CanComputeOutput() on those (Node/Output) pairs
+	for (int32 k = 0; k < NumInputRequirements; ++k)
+	{
+		const FString& InputName = InputRequirements[k].InputName;
+
+		// find the connection for this input
+		FConnection Connection;
+		EGeometryFlowResult FoundResult = FindConnectionForInput(NodeHandle, InputName, Connection);
+
+		if (FoundResult == EGeometryFlowResult::ConnectionDoesNotExist)
+		{
+			continue;
+		}
+
+
+		TFuture<bool> Future = Async(GeometryFlowGraphAsyncExecTarget,
+			[this, &InputName, Connection, Visitor]()
+			{
+				if (Visitor(Connection.FromNode)) return true;
+
+				// recursively look upstream from this Input via Connection
+				const bool UpstreamResult = VistDependencies(Connection.FromNode, Visitor);
+				return UpstreamResult;
+			});
+
+		AsyncFutures.Emplace(MoveTemp(Future));
+	}
+
+	bool bVisitorTerminated = false;
+	for (TFuture<bool>& Future : AsyncFutures)
+	{
+		Future.Wait();
+		bVisitorTerminated = bVisitorTerminated || Future.Get();
+	}
+
+	return bVisitorTerminated;
+}
+
 
 
 
@@ -420,3 +594,73 @@ FString FGraph::DebugDumpGraph(TFunction<bool(TSafeSharedPtr<FNode>)> IncludeNod
 
 	return Out;
 }
+
+
+void FGraph::Serialize(FArchive& Ar)
+{
+
+
+	if (Ar.IsSaving())
+	{
+		float CurrentVersion = GetVersionID();
+		Ar << CurrentVersion;
+	}
+	else
+	{
+		float ArchievedVersion;
+		Ar << ArchievedVersion;
+
+		if (!ensureMsgf(ArchievedVersion == GetVersionID(), TEXT("Unable to load Geometry Flow Graph Version %f"), ArchievedVersion))
+		{
+			// any future version of this graph will have to provide some mechanism for loading older versions
+			return;
+		}
+	}
+	// Data for serialization:  NodeCounter, Connections, AllNodes (note, AllNodeLocks aren't serialized but are reconstructed during IsLoading)
+	Ar << NodeCounter;
+
+	Ar << Connections;
+
+	if (Ar.IsSaving())
+	{
+		int32 NumNodes = AllNodes.Num();
+		Ar << NumNodes;
+
+		for (TPair<FHandle, FNodeInfo>& Pair : AllNodes)
+		{
+			FHandle Handle = Pair.Key;
+			FNodeInfo& NodeInfo = Pair.Value;
+			FName NodeTypeName = NodeInfo.Node->GetType();
+			FString Identifier = NodeInfo.Node->GetIdentifier();
+			checkSlow(NodeTypeName == NodeInfo.NodeTypeName);
+			Ar << Handle;
+			Ar << NodeTypeName;
+			Ar << NodeInfo.CachingStrategy;
+			Ar << Identifier;
+
+			NodeInfo.Node->Serialize(Ar);
+		}
+	}
+	else if (Ar.IsLoading())
+	{
+		int32 NumNodes = 0;
+		Ar << NumNodes;
+
+		for (int32 i = 0; i < NumNodes; ++i)
+		{
+			FHandle Handle;
+			FName NodeTypeName;
+			ENodeCachingStrategy CachingStrategy;
+			FString Identifier;
+
+			Ar << Handle;
+			Ar << NodeTypeName;
+			Ar << CachingStrategy;
+			Ar << Identifier;
+
+			const ENodeAddResult AddResult = AddNodeOfType(NodeTypeName, Handle, Identifier, CachingStrategy); // automatically populates AllNodes and AllNodeLocks
+			checkSlow(AddResult == ENodeAddResult::Success);
+			ApplyToNode(Handle, [&Ar](FNode& Node) { Node.Serialize(Ar); });
+		}
+	}
+}
\ No newline at end of file
diff --git a/Engine/Plugins/Experimental/GeometryFlow/Source/GeometryFlowCore/Private/GeometryFlowNode.cpp b/Engine/Plugins/Experimental/GeometryFlow/Source/GeometryFlowCore/Private/GeometryFlowNode.cpp
--- a/Engine/Plugins/Experimental/GeometryFlow/Source/GeometryFlowCore/Private/GeometryFlowNode.cpp
+++ b/Engine/Plugins/Experimental/GeometryFlow/Source/GeometryFlowCore/Private/GeometryFlowNode.cpp
@@ -32,7 +32,7 @@ void FNode::AddOutput(FString Name, TUniquePtr<INodeOutput>&& Output)
 EGeometryFlowResult FNode::GetInputType(FString Name, int32& TypeOut) const
 {
 	int32 Index = NodeInputs.IndexOfByPredicate([&](const FNodeInputInfo& Input) { return Input.Name == Name; });
-	if (ensure(Index != INDEX_NONE))
+	if (Index != INDEX_NONE)
 	{
 		TypeOut = NodeInputs[Index].Input->GetDataType();
 		return EGeometryFlowResult::Ok;
@@ -53,7 +53,7 @@ FNodeInputFlags FNode::GetInputFlags(FString InputName) const
 EGeometryFlowResult FNode::GetOutputType(FString Name, int32& TypeOut) const
 {
 	int32 Index = NodeOutputs.IndexOfByPredicate([&](const FNodeOutputInfo& Output) { return Output.Name == Name; });
-	if (ensure(Index != INDEX_NONE))
+	if (Index != INDEX_NONE)
 	{
 		TypeOut = NodeOutputs[Index].Output->GetDataType();
 		return EGeometryFlowResult::Ok;
diff --git a/Engine/Plugins/Experimental/GeometryFlow/Source/GeometryFlowCore/Private/GometryFlowCoreNodeRegistration.cpp b/Engine/Plugins/Experimental/GeometryFlow/Source/GeometryFlowCore/Private/GometryFlowCoreNodeRegistration.cpp
new file mode 100644
--- /dev/null
+++ b/Engine/Plugins/Experimental/GeometryFlow/Source/GeometryFlowCore/Private/GometryFlowCoreNodeRegistration.cpp
@@ -0,0 +1,52 @@
+// Copyright Epic Games, Inc. All Rights Reserved.
+
+#include "GeometryFlowNodeFactory.h"
+#include "GeometryFlowCoreNodes.h"
+
+#include "BaseNodes/SwitchNode.h"
+#include "MathNodes/ArithmeticNodes.h"
+
+
+namespace UE
+{
+namespace GeometryFlow
+{
+#define GEOMETRY_FLOW_QUOTE(x) #x
+
+// register node of type F{TypeName}SourceNode with FName {TypeName}Source
+#define GEOMETRYFLOW_REGISTER_BASIC_TYPES(TypeName) \
+	FNodeFactory::GetInstance().RegisterType<F##TypeName##SourceNode>(FString(GEOMETRY_FLOW_QUOTE(TypeName##Source)), \
+																	  CategoryName);
+// register node of type F{Typename}StructSourceNode with FName {TypeName}StructSource
+#define GEOMETRYFLOW_REGISTER_BASIC_USTRUCT_TYPES(TypeName) \
+	FNodeFactory::GetInstance().RegisterType<F##TypeName##Struct##SourceNode>(FString(GEOMETRY_FLOW_QUOTE(TypeName##Struct##Source)), \
+																	  CategoryName);
+
+	// Register Core Nodes with the node factory
+	void RegisterCoreNodes()
+	{	
+		FString CategoryName("Basic Types");
+
+		GEOMETRYFLOW_REGISTER_BASIC_TYPES(Int32);
+		GEOMETRYFLOW_REGISTER_BASIC_TYPES(Float);
+		GEOMETRYFLOW_REGISTER_BASIC_TYPES(Double);
+		GEOMETRYFLOW_REGISTER_BASIC_TYPES(Vector3f);
+		GEOMETRYFLOW_REGISTER_BASIC_TYPES(Vector3d);
+		GEOMETRYFLOW_REGISTER_BASIC_TYPES(Name);
+
+		GEOMETRYFLOW_REGISTER_BASIC_USTRUCT_TYPES(Int32);
+		GEOMETRYFLOW_REGISTER_BASIC_USTRUCT_TYPES(Float);
+		GEOMETRYFLOW_REGISTER_BASIC_USTRUCT_TYPES(Double);
+		GEOMETRYFLOW_REGISTER_BASIC_USTRUCT_TYPES(Vector3f);
+		GEOMETRYFLOW_REGISTER_BASIC_USTRUCT_TYPES(Vector3d);
+		GEOMETRYFLOW_REGISTER_BASIC_USTRUCT_TYPES(Name);
+
+		FNodeFactory::GetInstance().RegisterType< FAddFloatNode >(FString("AddFloat"), CategoryName);
+	}
+
+
+#undef GEOMETRYFLOW_REGISTER_BASIC_TYPES_WUI
+#undef GEOMETRYFLOW_REGISTER_BASIC_TYPES
+#undef GEOMETRY_FLOW_QUOTE
+}
+}
\ No newline at end of file
diff --git a/Engine/Plugins/Experimental/GeometryFlow/Source/GeometryFlowCore/Public/BaseNodes/BinaryOpLambdaNode.h b/Engine/Plugins/Experimental/GeometryFlow/Source/GeometryFlowCore/Public/BaseNodes/BinaryOpLambdaNode.h
--- a/Engine/Plugins/Experimental/GeometryFlow/Source/GeometryFlowCore/Public/BaseNodes/BinaryOpLambdaNode.h
+++ b/Engine/Plugins/Experimental/GeometryFlow/Source/GeometryFlowCore/Public/BaseNodes/BinaryOpLambdaNode.h
@@ -3,10 +3,7 @@
 #pragma once
 
 #include "GeometryFlowNode.h"
-
-namespace UE::GeometryFlow { template <typename T, int32 DataTypeIdentifier> class TMovableData; }
-namespace UE::GeometryFlow { template <typename T, int32 StorageTypeIdentifier> class TBasicNodeInput; }
-namespace UE::GeometryFlow { template <typename T, int32 StorageTypeIdentifier> class TBasicNodeOutput; }
+#include "GeometryFlowMovableData.h"
 
 
 namespace UE
diff --git a/Engine/Plugins/Experimental/GeometryFlow/Source/GeometryFlowCore/Public/BaseNodes/SwitchNode.h b/Engine/Plugins/Experimental/GeometryFlow/Source/GeometryFlowCore/Public/BaseNodes/SwitchNode.h
--- a/Engine/Plugins/Experimental/GeometryFlow/Source/GeometryFlowCore/Public/BaseNodes/SwitchNode.h
+++ b/Engine/Plugins/Experimental/GeometryFlow/Source/GeometryFlowCore/Public/BaseNodes/SwitchNode.h
@@ -3,10 +3,7 @@
 #pragma once
 
 #include "GeometryFlowNode.h"
-
-namespace UE::GeometryFlow { template <typename T, int32 DataTypeIdentifier> class TMovableData; }
-namespace UE::GeometryFlow { template <typename T, int32 StorageTypeIdentifier> class TBasicNodeInput; }
-namespace UE::GeometryFlow { template <typename T, int32 StorageTypeIdentifier> class TBasicNodeOutput; }
+#include "GeometryFlowMovableData.h"
 
 
 namespace UE
diff --git a/Engine/Plugins/Experimental/GeometryFlow/Source/GeometryFlowCore/Public/BaseNodes/TransferNode.h b/Engine/Plugins/Experimental/GeometryFlow/Source/GeometryFlowCore/Public/BaseNodes/TransferNode.h
--- a/Engine/Plugins/Experimental/GeometryFlow/Source/GeometryFlowCore/Public/BaseNodes/TransferNode.h
+++ b/Engine/Plugins/Experimental/GeometryFlow/Source/GeometryFlowCore/Public/BaseNodes/TransferNode.h
@@ -3,10 +3,7 @@
 #pragma once
 
 #include "GeometryFlowNode.h"
-
-namespace UE::GeometryFlow { template <typename T, int32 DataTypeIdentifier> class TMovableData; }
-namespace UE::GeometryFlow { template <typename T, int32 StorageTypeIdentifier> class TBasicNodeInput; }
-namespace UE::GeometryFlow { template <typename T, int32 StorageTypeIdentifier> class TBasicNodeOutput; }
+#include "GeometryFlowMovableData.h"
 
 
 namespace UE
diff --git a/Engine/Plugins/Experimental/GeometryFlow/Source/GeometryFlowCore/Public/BaseNodes/TransformerWithSettingsNode.h b/Engine/Plugins/Experimental/GeometryFlow/Source/GeometryFlowCore/Public/BaseNodes/TransformerWithSettingsNode.h
--- a/Engine/Plugins/Experimental/GeometryFlow/Source/GeometryFlowCore/Public/BaseNodes/TransformerWithSettingsNode.h
+++ b/Engine/Plugins/Experimental/GeometryFlow/Source/GeometryFlowCore/Public/BaseNodes/TransformerWithSettingsNode.h
@@ -3,10 +3,7 @@
 #pragma once
 
 #include "GeometryFlowNode.h"
-
-namespace UE::GeometryFlow { template <typename T, int32 DataTypeIdentifier> class TMovableData; }
-namespace UE::GeometryFlow { template <typename T, int32 StorageTypeIdentifier> class TBasicNodeInput; }
-namespace UE::GeometryFlow { template <typename T, int32 StorageTypeIdentifier> class TBasicNodeOutput; }
+#include "GeometryFlowMovableData.h"
 
 
 namespace UE
diff --git a/Engine/Plugins/Experimental/GeometryFlow/Source/GeometryFlowCore/Public/GeometryFlowCoreNodes.h b/Engine/Plugins/Experimental/GeometryFlow/Source/GeometryFlowCore/Public/GeometryFlowCoreNodes.h
--- a/Engine/Plugins/Experimental/GeometryFlow/Source/GeometryFlowCore/Public/GeometryFlowCoreNodes.h
+++ b/Engine/Plugins/Experimental/GeometryFlow/Source/GeometryFlowCore/Public/GeometryFlowCoreNodes.h
@@ -2,122 +2,96 @@
 
 #pragma once
 
-#include "GeometryBase.h"
-#include "GeometryFlowNode.h"
+#include "GeometryFlowTSourceNode.h"
+#include "GeometryFlowTypes.h"
 
-namespace UE::GeometryFlow { template <typename T, int32 DataTypeIdentifier> class TMovableData; }
-namespace UE::GeometryFlow { template <typename T, int32 StorageTypeIdentifier> class TBasicNodeInput; }
-namespace UE::GeometryFlow { template <typename T, int32 StorageTypeIdentifier> class TBasicNodeOutput; }
+#include "GeometryFlowCoreNodes.generated.h"
 
-
-namespace UE
-{
-namespace GeometryFlow
+USTRUCT()
+struct GEOMETRYFLOWCORE_API FInt32Setting
 {
+	GENERATED_USTRUCT_BODY()
+	DECLARE_GEOMETRYFLOW_DATA_TYPE_IDENTIFIER(UE::GeometryFlow::EDataTypes::Integer);
 
-using namespace UE::Geometry;
-
-
-template<typename T, int32 StorageTypeIdentifier>
-class TSourceNode : public FNode
-{
-protected:
-	using DataType = TMovableData<T, StorageTypeIdentifier>;
-
-	TSafeSharedPtr<DataType> Value;
-
-public:
-	using CppType = T;
-	static constexpr int DataTypeIdentifier = StorageTypeIdentifier;
-
-	static const FString OutParamValue() { return TEXT("Value"); }
-
-public:
-	TSourceNode()
-	{
-		T InitialValue;
-		Value = MakeSafeShared<DataType>(InitialValue);
-		AddOutput(OutParamValue(), MakeUnique<TBasicNodeOutput<T, StorageTypeIdentifier>>());
-
-		UpdateSourceValue(InitialValue);
-	}
-
-	void UpdateSourceValue(const T& NewValue)
-	{
-		Value->SetData(NewValue);
-		SetOutput(OutParamValue(), Value);
-	}
-
-	virtual void CollectRequirements(const TArray<FString>& Outputs, TArray<FEvalRequirement>& RequiredInputsOut) override
-	{
-		return;
-	}
-
-	virtual void CollectAllRequirements(TArray<FEvalRequirement>& RequiredInputsOut) override
-	{
-		return;
-	}
-
-	virtual void Evaluate(
-		const FNamedDataMap& DatasIn,
-		FNamedDataMap& DatasOut,
-		TUniquePtr<FEvaluationInfo>& EvaluationInfo) override
-	{
-		if ( ensure(DatasOut.Contains(OutParamValue())) )
-		{
-			DatasOut.SetData(OutParamValue(), GetOutput(OutParamValue()));
-		}
-	}
+	UPROPERTY(EditAnywhere, Category = "Geometry Flow")
+	int32 Value = 0;	
 };
 
+USTRUCT()
+struct GEOMETRYFLOWCORE_API FFloatSetting
+{
+	GENERATED_USTRUCT_BODY()
+	DECLARE_GEOMETRYFLOW_DATA_TYPE_IDENTIFIER(UE::GeometryFlow::EDataTypes::Float);
 
+	UPROPERTY(EditAnywhere, Category = "Geometry Flow")
+	float Value = 0;	
+};
 
+USTRUCT()
+struct GEOMETRYFLOWCORE_API FDoubleSetting
+{
+	GENERATED_USTRUCT_BODY()
+	DECLARE_GEOMETRYFLOW_DATA_TYPE_IDENTIFIER(UE::GeometryFlow::EDataTypes::Double);
 
-/**
- * This macro declares a MovableData, Node Input/Output, and TSourceNode type for a given C++ type/identifier
- * 
- * eg GEOMETRYFLOW_DECLARE_BASIC_TYPES(Int32, int, (int)EDataTypes::Integer) will declare/define:
- *     - TMovableData type named FDataInt32
- *     - TBasicNodeInput type named FInt32Input
- *     - TBasicNodeOutput type named FInt32Output
- *     - TSourceNode type named FInt32SourceNode
- *
- */
-#define GEOMETRYFLOW_DECLARE_BASIC_TYPES(TypeName, CppType, TypeIdentifier) \
-	typedef TMovableData<CppType, TypeIdentifier> FData##TypeName; \
-	typedef TBasicNodeInput<CppType, TypeIdentifier> F##TypeName##Input; \
-	typedef TBasicNodeOutput<CppType, TypeIdentifier> F##TypeName##Output; \
-	typedef TSourceNode<CppType, TypeIdentifier> F##TypeName##SourceNode; 
-
-
+	UPROPERTY(EditAnywhere, Category = "Geometry Flow")
+	double Value = 0;	
+};	
 
 
+USTRUCT()
+struct GEOMETRYFLOWCORE_API FVector3fSetting
+{
+	GENERATED_USTRUCT_BODY()
+	DECLARE_GEOMETRYFLOW_DATA_TYPE_IDENTIFIER(UE::GeometryFlow::EDataTypes::Vector3f);
 
-/**
- * This macro defines two classes for a given C++ class:
- *   - TMovableData typedef named FData[XYZ]Settings
- *   - TSourceNode typedef named F[XYZ]SettingsSourceNode
- *
- * The assumption is the class has an integer static/constexpr member SettingsType::DataTypeIdentifier that defines the type integer
- */
-#define GEOMETRYFLOW_DECLARE_SETTINGS_TYPES(CppType, ReadableName) \
-	typedef TMovableData<CppType, CppType::DataTypeIdentifier> FData##ReadableName##Settings; \
-	typedef TSourceNode<CppType, CppType::DataTypeIdentifier> F##ReadableName##SettingsSourceNode;
+	UPROPERTY(EditAnywhere, Category = "Geometry Flow")
+	FVector3f Value = FVector3f(0, 0, 0);	
+};	
 
+USTRUCT()
+struct GEOMETRYFLOWCORE_API FVector3dSetting
+{
+	GENERATED_USTRUCT_BODY()
+	DECLARE_GEOMETRYFLOW_DATA_TYPE_IDENTIFIER(UE::GeometryFlow::EDataTypes::Vector3d);
 
+	UPROPERTY(EditAnywhere, Category = "Geometry Flow")
+	FVector3d Value = FVector3d(0, 0, 0);	
+};	
 
 
+USTRUCT()
+struct GEOMETRYFLOWCORE_API FNameSetting
+{
+	GENERATED_USTRUCT_BODY()
+	DECLARE_GEOMETRYFLOW_DATA_TYPE_IDENTIFIER(UE::GeometryFlow::EDataTypes::Name);
 
+	UPROPERTY(EditAnywhere, Category = "Geometry Flow")
+	FName Value;	
+};	
 
+namespace UE
+{
+namespace GeometryFlow
+{
 /**
  * Declare basic types for math
  */
-GEOMETRYFLOW_DECLARE_BASIC_TYPES(Int32, int, (int)EDataTypes::Integer)
-GEOMETRYFLOW_DECLARE_BASIC_TYPES(Float, float, (int)EDataTypes::Float)
-GEOMETRYFLOW_DECLARE_BASIC_TYPES(Double, double, (int)EDataTypes::Double)
-GEOMETRYFLOW_DECLARE_BASIC_TYPES(Vector3f, FVector3f, (int)EDataTypes::Vector3f)
-GEOMETRYFLOW_DECLARE_BASIC_TYPES(Vector3d, FVector3d, (int)EDataTypes::Vector3d)
-GEOMETRYFLOW_DECLARE_BASIC_TYPES(Name, FName, (int)EDataTypes::Name)
+GEOMETRYFLOW_DECLARE_BASIC_TYPES(Int32, int, (int)EDataTypes::Integer, 1)
+GEOMETRYFLOW_DECLARE_BASIC_TYPES(Float, float, (int)EDataTypes::Float, 1)
+GEOMETRYFLOW_DECLARE_BASIC_TYPES(Double, double, (int)EDataTypes::Double, 1)
+GEOMETRYFLOW_DECLARE_BASIC_TYPES(Vector3f, FVector3f, (int)EDataTypes::Vector3f, 1)
+GEOMETRYFLOW_DECLARE_BASIC_TYPES(Vector3d, FVector3d, (int)EDataTypes::Vector3d, 1)
+GEOMETRYFLOW_DECLARE_BASIC_TYPES(Name, FName, (int)EDataTypes::Name, 1)
+
+/**
+* basic math types wrapped in USTRUCTs so they can be exposed as details
+*/
+GEOMETRYFLOW_DECLARE_BASIC_USTRUCT_TYPES(Int32, FInt32Setting, (int)EDataTypes::IntegerStruct, 1)
+GEOMETRYFLOW_DECLARE_BASIC_USTRUCT_TYPES(Float, FFloatSetting, (int)EDataTypes::FloatStruct, 1)
+GEOMETRYFLOW_DECLARE_BASIC_USTRUCT_TYPES(Double, FDoubleSetting, (int)EDataTypes::DoubleStruct, 1)
+GEOMETRYFLOW_DECLARE_BASIC_USTRUCT_TYPES(Vector3f, FVector3fSetting, (int)EDataTypes::Vector3fStruct, 1)
+GEOMETRYFLOW_DECLARE_BASIC_USTRUCT_TYPES(Vector3d, FVector3dSetting, (int)EDataTypes::Vector3dStruct, 1)
+GEOMETRYFLOW_DECLARE_BASIC_USTRUCT_TYPES(Name, FNameSetting, (int)EDataTypes::NameStruct, 1)
 
 
 }	// end namespace GeometryFlow
diff --git a/Engine/Plugins/Experimental/GeometryFlow/Source/GeometryFlowCore/Public/GeometryFlowGraph.h b/Engine/Plugins/Experimental/GeometryFlow/Source/GeometryFlowCore/Public/GeometryFlowGraph.h
--- a/Engine/Plugins/Experimental/GeometryFlow/Source/GeometryFlowCore/Public/GeometryFlowGraph.h
+++ b/Engine/Plugins/Experimental/GeometryFlow/Source/GeometryFlowCore/Public/GeometryFlowGraph.h
@@ -3,7 +3,9 @@
 #pragma once
 
 #include "GeometryFlowNode.h"
+#include "GeometryFlowNodeFactory.h"
 #include "Util/ProgressCancel.h"
+#include "Serialization/Archive.h"
 
 class FProgressCancel;
 
@@ -35,15 +37,40 @@ enum class ENodeCachingStrategy
 
 class GEOMETRYFLOWCORE_API FGraph
 {
+
+	// version used in serialization 
+	static constexpr float SerializationVersion = 1.0;
+
 public:
 
 	~FGraph();
 
+	float GetVersionID() const { return FGraph::SerializationVersion;}
+
 	struct FHandle
 	{
 		static const int32 InvalidHandle = -1;
 		int32 Identifier = InvalidHandle;
 		bool operator==(const FHandle& OtherHandle) const { return Identifier == OtherHandle.Identifier; }
+
+		/**
+		* Serialization operator for FHandle
+		*
+		* @parm Ar Archive to serialize with.
+		* @param H  Handle to serialize.
+		* @returns Passing down serializing archive
+		*/
+		friend FArchive& operator<<(FArchive& Ar, FHandle& H)
+		{
+			H.Serialize(Ar);
+			return Ar;
+		}
+
+		/** Serialize FHandle to an archive*/
+		void Serialize(FArchive& Ar)
+		{
+			Ar << Identifier;
+		}
 	};
 
 	friend uint32 GetTypeHash(FHandle Handle);
@@ -54,27 +81,129 @@ public:
 		FString FromOutput;
 		FHandle ToNode;
 		FString ToInput;
+
+		/**
+		* Serialization operator for FConnection
+		*
+		* @parm Ar Archive to serialize with.
+		* @param C  Connection to serialize.
+		* @returns Passing down serializing archive
+		*/
+		friend FArchive& operator<<(FArchive& Ar, FConnection& C)
+		{
+			C.Serialize(Ar);
+			return Ar;
+		}
+
+		/** Serialize FHandle to an archive*/
+		void Serialize(FArchive& Ar)
+		{
+			Ar << FromNode;
+			Ar << FromOutput;
+			Ar << ToNode;
+			Ar << ToInput;
+		}
+
+		bool operator==(const FConnection& Other) const
+		{
+			return (FromNode == Other.FromNode)
+				&& (FromOutput == Other.FromOutput)
+				&& (ToNode == Other.ToNode)
+				&& (ToInput == Other.ToInput);
+		}
 	};
 
-	template<typename NodeType>
+	// return @false if the node type is not registered
+	template <typename NodeType>
+	bool CanAddNodeOfType()
+	{
+		return CanAddNodeOfType(NodeType::StaticType());
+	}
+
+	// Add a node of specified type to the graph.  An invalid FHandle will be returned if failed (e.g. node type not registered)
+	template <typename NodeType>
+	FHandle AddNodeOfType(const FString& Identifier = FString(""),
+		ENodeCachingStrategy CachingStrategy = ENodeCachingStrategy::Default)
+	{
+		return AddNodeOfType(NodeType::StaticType(), Identifier, CachingStrategy);
+	}
+
+	// Add a node of specified type to the graph.  An invalid FHandle will be returned if failed (e.g. node type not registered)
 	FHandle AddNodeOfType(
+		const FName TypeName,
 		const FString& Identifier = FString(""),
-		ENodeCachingStrategy CachingStrategy = ENodeCachingStrategy::Default )
+		ENodeCachingStrategy CachingStrategy = ENodeCachingStrategy::Default)
 	{
-		FNodeInfo NewNodeInfo;
-		NewNodeInfo.Node = MakeShared<NodeType, ESPMode::ThreadSafe>();
-		NewNodeInfo.Node->SetIdentifier(Identifier);
-		NewNodeInfo.CachingStrategy = CachingStrategy;
-		FHandle Handle = { NodeCounter++ };
-		AllNodes.Add(Handle, NewNodeInfo);
+		FHandle Handle;
+		Handle.Identifier= NodeCounter + 1;
+
+		if (AddNodeOfType(TypeName, Handle, Identifier, CachingStrategy) == ENodeAddResult::Success)
+		{
+			NodeCounter++;
+		}
+		else
+		{
+			Handle = FHandle(); // invalid handle.
+		}
 
-		AllNodeLocks.Add(Handle, MakeSafeShared<FRWLock>());
 		return Handle;
 	}
 
+	// Note: once removed, any external handles to this node will be stale. 
+	void RemoveNode(const FHandle Handle)
+	{
+		RemoveNodeConnections(Handle);
+
+		AllNodes.Remove(Handle);
+		AllNodeLocks.Remove(Handle);
+	}
+
+	// returns true if DependentNode is dependent on IndependentNode in the graph.
+	// special cases:
+	// will return false if either handle is invalid,
+	// if both handles are the same (and valid) it will return true.
+	bool IsDependent(FHandle DependentNode, FHandle IndependentNode) const;
+
+	// returns EGeometryFlowResult::Ok if the Connection can be added, otherwise returns the reason for failure
+	EGeometryFlowResult CanAddConnection(FHandle FromNode, FString FromOutput, FHandle ToNode, FString ToInput) const;
+
 	EGeometryFlowResult AddConnection(FHandle FromNode, FString FromOutput, FHandle ToNode, FString ToInput);
 
 	EGeometryFlowResult InferConnection(FHandle FromNode, FHandle ToNode);
+	
+	// find all the input connections for the specified node.
+	EGeometryFlowResult FindInputConnections(FHandle ToNode, TArray<FGraph::FConnection>& ConnectionsOut) const;
+
+	// remove input connections to the specifed node
+	void RemoveNodeInputConnections(FHandle NodeHandle)
+	{
+		int32 index = Connections.IndexOfByPredicate([&](const FConnection& Conn) { return (Conn.ToNode == NodeHandle); });
+		while (index != INDEX_NONE)
+		{
+			Connections.RemoveAtSwap(index);
+			index = Connections.IndexOfByPredicate([&](const FConnection& Conn) { return (Conn.ToNode == NodeHandle); });
+		}
+	}
+
+	// remove all output connections to the specified node.
+	void RemoveNodeOutputConnections(FHandle NodeHandle)
+	{
+		int32 index = Connections.IndexOfByPredicate([&](const FConnection& Conn) { return (Conn.FromNode == NodeHandle); });
+		while (index != INDEX_NONE)
+		{
+			Connections.RemoveAtSwap(index);
+			index = Connections.IndexOfByPredicate([&](const FConnection& Conn) { return (Conn.FromNode == NodeHandle); });
+		}
+	}
+
+	// remove all connections to the specified node.
+	void RemoveNodeConnections(FHandle NodeHandle)
+	{
+		RemoveNodeInputConnections(NodeHandle);
+		RemoveNodeOutputConnections(NodeHandle);
+	}
+
+	EGeometryFlowResult RemoveConnectionForInput(FHandle ToNode, FString ToInput);
 
 	TSet<FHandle> GetSourceNodes() const;
 	TSet<FHandle> GetNodesWithNoConnectedInputs() const;
@@ -95,6 +224,15 @@ public:
 		return Result;
 	}
 
+	bool CanEvaluate(
+		FHandle Node,
+		FString OutputName
+	)
+	{
+		return CanComputeOutput(Node, OutputName);
+	}
+
+	// Applies the functor to the node as derived type if possible.
 	template<typename NodeType>
 	EGeometryFlowResult ApplyToNodeOfType(
 		FHandle NodeHandle, 
@@ -110,15 +248,90 @@ public:
 		return EGeometryFlowResult::NodeDoesNotExist;
 	}
 
+	// Applies the functor to the node as FNode.
+	EGeometryFlowResult ApplyToNode(
+		FHandle NodeHandle,
+		TFunctionRef<void(FNode&)> ApplyFunc)
+	{
+		TSafeSharedPtr<FNode> FoundNode = FindNode(NodeHandle);
+		if (FoundNode)
+		{
+			ApplyFunc(*FoundNode);
+			return EGeometryFlowResult::Ok;
+		}
+		return EGeometryFlowResult::NodeDoesNotExist;
+	}
 
+	EGeometryFlowResult GetInputTypeForNode(FHandle NodeHandle, FString InputName, int32& Type) const;
+	EGeometryFlowResult GetOutputTypeForNode(FHandle NodeHandle, FString OutputName, int32& Type) const;
 
 	void ConfigureCachingStrategy(ENodeCachingStrategy NewStrategy);
 	EGeometryFlowResult SetNodeCachingStrategy(FHandle NodeHandle, ENodeCachingStrategy Strategy);
 
 	FString DebugDumpGraph(TFunction<bool(TSafeSharedPtr<FNode>)> IncludeNodeFn) const;
 
+	/**
+	* Serialization operator for FGraph.
+	*
+	* @param Ar Archive to serialize with.
+	* @param Graph Graph to serialize.
+	* @returns Passing down serializing archive.
+	*/
+	friend GEOMETRYFLOWCORE_API FArchive& operator<<(FArchive& Ar, FGraph& Graph)
+	{
+		Graph.Serialize(Ar);
+		return Ar;
+	}
+
+	/** Serialize this graph to an archive. */
+	void Serialize(FArchive& Ar);
+
+
 protected:
 
+	enum class ENodeAddResult
+	{
+		Success = 0,
+		Failed_HandleExists = 1,
+		Failed_UnregisteredType = 2,
+	};
+
+
+	bool CanAddNodeOfType(const FName TypeName)
+	{
+		const FNodeFactory& NodeFactory = FNodeFactory::GetInstance();
+
+		return NodeFactory.CanMake(TypeName);
+	}
+
+	ENodeAddResult AddNodeOfType(const FName TypeName,
+		const FHandle Handle,
+		const FString& Identifier,
+		ENodeCachingStrategy CachingStrategy)
+	{
+		if (AllNodes.Contains(Handle))
+		{
+			return ENodeAddResult::Failed_HandleExists;
+		}
+
+		const FNodeFactory& NodeFactory = FNodeFactory::GetInstance();
+
+		if (!NodeFactory.CanMake(TypeName))
+		{
+			return ENodeAddResult::Failed_UnregisteredType;
+		}
+
+		FNodeInfo NewNodeInfo;
+		NewNodeInfo.Node = TSharedPtr<FNode, ESPMode::ThreadSafe>(NodeFactory.CreateNodeOfType(TypeName).Release());
+		NewNodeInfo.NodeTypeName = TypeName;
+		NewNodeInfo.Node->SetIdentifier(Identifier);
+		NewNodeInfo.CachingStrategy = CachingStrategy;
+
+		AllNodes.Add(Handle, NewNodeInfo);
+		AllNodeLocks.Add(Handle, MakeSafeShared<FRWLock>());
+		return ENodeAddResult::Success;
+	}
+
 	template<typename T>
 	EGeometryFlowResult EvaluateResultInternal(
 		FHandle Node,
@@ -169,6 +382,7 @@ protected:
 	struct FNodeInfo
 	{
 		TSafeSharedPtr<FNode> Node;
+		FName  NodeTypeName = FName("UnknowNodeType");
 		ENodeCachingStrategy CachingStrategy = ENodeCachingStrategy::Default;
 	};
 
@@ -176,8 +390,7 @@ protected:
 	TMap<FHandle, TSafeSharedPtr<FRWLock>> AllNodeLocks;
 
 	TSafeSharedPtr<FNode> FindNode(FHandle Handle) const;
-	EGeometryFlowResult GetInputTypeForNode(FHandle NodeHandle, FString InputName, int32& Type) const;
-	EGeometryFlowResult GetOutputTypeForNode(FHandle NodeHandle, FString OutputName, int32& Type) const;
+	
 	ENodeCachingStrategy GetCachingStrategyForNode(FHandle NodeHandle) const;
 	TSafeSharedPtr<FRWLock> FindNodeLock(FHandle Handle) const;
 
@@ -193,6 +406,25 @@ protected:
 		TUniquePtr<FEvaluationInfo>& EvaluationInfo,
 		bool bStealOutputData = false);
 
+	// retrun @true if all the graph connections needed for this computation exist
+	bool CanComputeOutput(
+		FHandle Node,
+		FString OutputName);
+
+	/**
+	* Visits upstream dependencies for the specified node and calls the Visitor functor (in parallel) on each.
+	* Terminates if Vistor returns true, or if the upstream graph is exhausted.
+	* @return true if terminated by the Visitor.
+	*/
+	bool VistDependencies(FHandle NodeHandle, TFunctionRef<bool(FGraph::FHandle)> Vistor) const;
+
+	bool IsInCycle(FHandle NodeHandle)
+	{
+		FHandle Start = NodeHandle;
+		auto Visitor = [this, Start](FGraph::FHandle UpstreamHandle) {return Start == UpstreamHandle; };
+		return VistDependencies(NodeHandle, Visitor);
+	}
+
 	FCriticalSection EvaluateLock;
 };
 
diff --git a/Engine/Plugins/Experimental/GeometryFlow/Source/GeometryFlowCore/Public/GeometryFlowGraphUtil.h b/Engine/Plugins/Experimental/GeometryFlow/Source/GeometryFlowCore/Public/GeometryFlowGraphUtil.h
--- a/Engine/Plugins/Experimental/GeometryFlow/Source/GeometryFlowCore/Public/GeometryFlowGraphUtil.h
+++ b/Engine/Plugins/Experimental/GeometryFlow/Source/GeometryFlowCore/Public/GeometryFlowGraphUtil.h
@@ -3,6 +3,7 @@
 #pragma once
 
 #include "GeometryFlowGraph.h"
+#include "GeometryFlowTSourceNode.h"
 
 namespace UE::GeometryFlow { template <typename T, int32 StorageTypeIdentifier> class TSourceNode; }
 
@@ -40,7 +41,7 @@ void UpdateSwitchNodeInputIndex(FGraph& Graph, FGraph::FHandle NodeHandle, int32
 template<typename SettingsType>
 void UpdateSettingsSourceNodeValue(FGraph& Graph, FGraph::FHandle NodeHandle, const SettingsType& NewSettings)
 {
-	using SettingsSourceNodeType = TSourceNode<SettingsType, SettingsType::DataTypeIdentifier>;
+	using SettingsSourceNodeType = TSourceNodeBase<SettingsType, SettingsType::DataTypeIdentifier>;
 
 	Graph.ApplyToNodeOfType<SettingsSourceNodeType>(NodeHandle, [&](SettingsSourceNodeType& Node)
 	{
diff --git a/Engine/Plugins/Experimental/GeometryFlow/Source/GeometryFlowCore/Public/GeometryFlowMovableData.h b/Engine/Plugins/Experimental/GeometryFlow/Source/GeometryFlowCore/Public/GeometryFlowMovableData.h
--- a/Engine/Plugins/Experimental/GeometryFlow/Source/GeometryFlowCore/Public/GeometryFlowMovableData.h
+++ b/Engine/Plugins/Experimental/GeometryFlow/Source/GeometryFlowCore/Public/GeometryFlowMovableData.h
@@ -3,6 +3,7 @@
 #pragma once
 
 #include "GeometryFlowNode.h"
+#include "Serialization/Archive.h"
 
 
 namespace UE
@@ -12,13 +13,20 @@ namespace GeometryFlow
 
 
 
-
+// This must be specialized for new data types.
+// see GeometryFlowTSourceNode.h for example
+template <typename DataType>
+struct TSerializationMethod
+{
+	static void Serialize(FArchive& Ar, DataType& T);
+};
 
 
 
 template<typename T, int32 DataTypeIdentifier>
 class TMovableData : public IData
 {
+	typedef IData Super;
 public:
 	T Data;
 
@@ -86,6 +94,12 @@ public:
 		Data = MoveTemp(DataIn);
 		IncrementTimestamp();
 	}
+
+	void Serialize(FArchive& Ar)
+	{
+		Ar << IData::Timestamp;
+		TSerializationMethod<T>::Serialize(Ar, Data);
+	}
 };
 
 
diff --git a/Engine/Plugins/Experimental/GeometryFlow/Source/GeometryFlowCore/Public/GeometryFlowNode.h b/Engine/Plugins/Experimental/GeometryFlow/Source/GeometryFlowCore/Public/GeometryFlowNode.h
--- a/Engine/Plugins/Experimental/GeometryFlow/Source/GeometryFlowCore/Public/GeometryFlowNode.h
+++ b/Engine/Plugins/Experimental/GeometryFlow/Source/GeometryFlowCore/Public/GeometryFlowNode.h
@@ -3,7 +3,8 @@
 #pragma once
 
 #include "GeometryFlowTypes.h"
-
+#include "UObject/StructOnScope.h"
+#include "Serialization/Archive.h"
 
 namespace UE::GeometryFlow { class FNode; }
 
@@ -293,11 +294,40 @@ public:
 	FNode& operator=(FNode&&) = default;
 	virtual ~FNode() = default;
 
+	//--- Derived nodes should implement these quasi-RTTI methods
+	//    using the provided macro GEOMETRYFLOW_NODE_RTTI()
+	static FName StaticType() { return "FNode"; }
+	virtual FName GetType() const { return "FNode"; }
+	virtual bool IsA(FName InType) const
+	{
+		return InType.ToString().Equals(StaticType().ToString());
+	}
+	//---
+	
 	void SetIdentifier(const FString& IdentifierIn) { Identifier = IdentifierIn; }
 	const FString& GetIdentifier() const { return Identifier; }
 
-	void AddInput(FString Name, TUniquePtr<INodeInput>&& Input, TSafeSharedPtr<IData> DefaultData = TSafeSharedPtr<IData>());
-	void AddOutput(FString Name, TUniquePtr<INodeOutput>&& Output);
+	// version number used for serialization - to allow for changes in node data.
+	virtual int32 GetVersionID() const { return -1; }
+	virtual void Serialize(FArchive& Ar)
+	{
+		if (Ar.IsLoading())
+		{
+			int32 SavedVersion;
+			Ar << SavedVersion;
+
+			const int32 ExpectedVersion = GetVersionID();
+			if (SavedVersion != ExpectedVersion)
+			{
+				UpdateVersion(Ar, SavedVersion);
+			}
+		}
+		else
+		{
+			int32 ThisVersion = GetVersionID();
+			Ar << ThisVersion;
+		}
+	}
 
 	EGeometryFlowResult GetInputType(FString Name, int32& TypeOut) const;
 	EGeometryFlowResult GetOutputType(FString Name, int32& TypeOut) const;
@@ -339,6 +369,11 @@ public:
 
 protected:
 
+	// derived types configure inputs and outputs in their constructor
+	void AddInput(FString Name, TUniquePtr<INodeInput>&& Input, TSafeSharedPtr<IData> DefaultData = TSafeSharedPtr<IData>());
+	void AddOutput(FString Name, TUniquePtr<INodeOutput>&& Output);
+
+
 	friend class FGraph;
 	friend class FGeometryFlowExecutor;
 
@@ -395,9 +430,37 @@ protected:
 
 	TSafeSharedPtr<IData> FindAndUpdateInputForEvaluate(const FString& InputName, const FNamedDataMap& DatasIn,
 		bool& bAccumModifiedOut, bool& bAccumValidOut);
-};
 
+protected:
 
+	// override this in derived node if the version number has been bumped
+	virtual void UpdateVersion(FArchive& Ar, int32 SavedVersion) {}
+};
+
+template<typename NodeType>
+NodeType* CastToNodePtr(FNode* Node)
+{
+	NodeType* Result = nullptr;
+	if (Node->IsA(NodeType::StaticType()))
+	{
+		Result = static_cast<NodeType*>(Node);
+	}
+	return Result;
+}
+
+#define GEOMETRYFLOW_NODE_RTTI(TYPE, SuperType)                                             \
+public:                                                                                     \
+	typedef SuperType   Super;                                                              \
+	static FName StaticType() { return #TYPE;}                                              \
+	virtual FName GetType() const override {return #TYPE;}                                  \
+	virtual bool  IsA(FName InType) const override                                          \
+		{ return InType.ToString().Equals(StaticType().ToString()) || Super::IsA(InType); }	\
+private:                                                      	
+
+#define GEOMETRYFLOW_NODE_INTERNAL(TYPE, VersionID, SuperType)            \
+public:                                                                   \
+	virtual int32 GetVersionID() const override { return VersionID;}      \
+    GEOMETRYFLOW_NODE_RTTI(TYPE, SuperType)                               \
 
 
 
diff --git a/Engine/Plugins/Experimental/GeometryFlow/Source/GeometryFlowCore/Public/GeometryFlowNodeFactory.h b/Engine/Plugins/Experimental/GeometryFlow/Source/GeometryFlowCore/Public/GeometryFlowNodeFactory.h
new file mode 100644
--- /dev/null
+++ b/Engine/Plugins/Experimental/GeometryFlow/Source/GeometryFlowCore/Public/GeometryFlowNodeFactory.h
@@ -0,0 +1,131 @@
+// Copyright Epic Games, Inc. All Rights Reserved.
+
+#pragma once
+
+#include "GeometryFlowNode.h"
+#include "CoreMinimal.h"
+
+namespace UE
+{
+namespace GeometryFlow
+{
+
+	// Singleton used to create nodes UE::GeometryFlow::FGraph, introduced to help with serialization.  
+	// this contains an map between FNames and function pointers that can create the named nodes.
+	class GEOMETRYFLOWCORE_API FNodeFactory
+	{
+
+	public:
+
+		struct FNodeFactoryInfo
+		{
+			typedef TFunction<TUniquePtr<FNode>() > FNodeTypeFactory;
+			
+			FString TypePrettyName;
+			FString TypeCategory;
+			FNodeTypeFactory  TypeFactory;
+		};
+
+
+		FNodeFactory(const FNodeFactory&) = delete;
+	
+		static FNodeFactory& GetInstance()
+		{
+			// Meyers style singleton
+			static FNodeFactory Instance;
+			return Instance;
+		};
+
+		// returns false if the NodeTypeName is already registered.
+		template <typename NodeType>
+		bool RegisterType(FString NodeTypePrettyName = FString(), FString TypeCategory = FString())
+		{
+			const FName NodeTypeName = NodeType::StaticType();
+			if (Factories.Contains(NodeTypeName))
+			{
+				return false;
+			}
+
+			if (NodeTypePrettyName.IsEmpty())
+			{
+				NodeTypePrettyName = NodeTypeName.ToString();
+			}
+			if (TypeCategory.IsEmpty())
+			{
+				TypeCategory = FString("Geometry Flow");
+			}
+
+			FNodeFactoryInfo NodeFactoryInfo= {NodeTypePrettyName, TypeCategory, []() { return MakeUnique<NodeType>(); } };
+
+			Factories.Add(NodeTypeName, NodeFactoryInfo);
+			return true;
+		}
+
+		template <typename NodeType>
+		bool CanMake() const
+		{
+			return CanMake(NodeType::StaticType());
+		}
+
+
+		bool CanMake(const FName NodeTypeName) const
+		{
+			return Factories.Contains(NodeTypeName);
+		}
+
+		// For unregistered types, this returns a null ptr.
+		TUniquePtr<FNode> CreateNodeOfType(const FName NodeTypeName) const 
+		{
+			TUniquePtr<FNode> Result(nullptr);
+			if (CanMake(NodeTypeName))
+			{
+				const FNodeFactoryInfo& FactoryInfo = Factories[NodeTypeName];
+				Result = FactoryInfo.TypeFactory();
+			}
+			return Result;
+		}
+
+		const TMap<const FName, FNodeFactory::FNodeFactoryInfo>& GetFactories() const { return Factories;}
+
+	protected:
+
+		TMap<const FName, FNodeFactoryInfo>  Factories;
+		
+		FNodeFactory()
+		{}
+	};
+
+	
+	// function that registers the core nodes defined in GeometryFlowCore (mostly POD source nodes)
+	void RegisterCoreNodes();
+}
+
+
+
+// by convention we register the nodes with an FName derived from the actual node class name
+#define GEOMETRYFLOW_QUOTE(x) #x
+
+// Register type F{ReadableName} with the FName {ReadableName}
+#define GEOMETRYFLOW_REGISTER_NODE_TYPE(ReadableName)                                                           \
+	FNodeFactory::GetInstance().RegisterType<F##ReadableName##Node>( FString(GEOMETRYFLOW_QUOTE(ReadableName)), \
+																	 CategoryName); 
+
+// Register Source Node type F{ReadableName}SourceNode  with FName {ReadableName}
+#define GEOMETRYFLOW_REGISTER_BASIC_TYPES_NODE(ReadableName)                                                         \
+	FNodeFactory::GetInstance().RegisterType<F##ReadableName##SourceNode>( FString(GEOMETRYFLOW_QUOTE(ReadableName)),\
+																		   CategoryName);
+
+// Register Setting Source Node type  F{ReadableName}SettingsSourceNode with FName {ReadableName}Settings
+#define GEOMETRYFLOW_REGISTER_SETTINGS_NODE_TYPE(ReadableName)                                                                         \
+	FNodeFactory::GetInstance().RegisterType<F##ReadableName##SettingsSourceNode>( FString(GEOMETRYFLOW_QUOTE(ReadableName##Settings)),\
+																				   CategoryName);
+
+// Register two nodes (an opperator and its settings).  F{ReadableName}Node with FName {ReadableName} and  F{ReadableName}SettingsSourceNode with FName {ReadableName}Settings
+#define GEOMETRYFLOW_REGISTER_NODE_AND_SETTINGS_NODE(ReadableName) \
+	FNodeFactory::GetInstance().RegisterType<F##ReadableName##Node>(FString(GEOMETRYFLOW_QUOTE(ReadableName)), \
+																	CategoryName);                             \
+																	                                           \
+	FNodeFactory::GetInstance().RegisterType<F##ReadableName##SettingsSourceNode>( FString(GEOMETRYFLOW_QUOTE(ReadableName##Settings)),\
+																				   CategoryName);
+
+}
diff --git a/Engine/Plugins/Experimental/GeometryFlow/Source/GeometryFlowCore/Public/GeometryFlowNodeUtil.h b/Engine/Plugins/Experimental/GeometryFlow/Source/GeometryFlowCore/Public/GeometryFlowNodeUtil.h
--- a/Engine/Plugins/Experimental/GeometryFlow/Source/GeometryFlowCore/Public/GeometryFlowNodeUtil.h
+++ b/Engine/Plugins/Experimental/GeometryFlow/Source/GeometryFlowCore/Public/GeometryFlowNodeUtil.h
@@ -3,10 +3,8 @@
 #pragma once
 
 #include "GeometryFlowNode.h"
+#include "GeometryFlowMovableData.h"
 
-namespace UE::GeometryFlow { template <typename T, int32 DataTypeIdentifier> class TMovableData; }
-namespace UE::GeometryFlow { template <typename T, int32 StorageTypeIdentifier> class TBasicNodeInput; }
-namespace UE::GeometryFlow { template <typename T, int32 StorageTypeIdentifier> class TBasicNodeOutput; }
 
 namespace UE
 {
diff --git a/Engine/Plugins/Experimental/GeometryFlow/Source/GeometryFlowCore/Public/GeometryFlowTSourceNode.h b/Engine/Plugins/Experimental/GeometryFlow/Source/GeometryFlowCore/Public/GeometryFlowTSourceNode.h
new file mode 100644
--- /dev/null
+++ b/Engine/Plugins/Experimental/GeometryFlow/Source/GeometryFlowCore/Public/GeometryFlowTSourceNode.h
@@ -0,0 +1,300 @@
+// Copyright Epic Games, Inc. All Rights Reserved.
+
+#pragma once
+
+#include "GeometryBase.h"
+#include "GeometryFlowNode.h"
+#include "GeometryFlowMovableData.h"
+#include "Templates/UniquePtr.h"
+
+
+namespace UE::GeometryFlow { template <typename T, int32 StorageTypeIdentifier> class TBasicNodeInput; }
+namespace UE::GeometryFlow { template <typename T, int32 StorageTypeIdentifier> class TBasicNodeOutput; }
+
+namespace UE
+{
+namespace GeometryFlow
+{
+
+using namespace UE::Geometry;
+
+class GEOMETRYFLOWCORE_API FSourceNodeBase : public FNode
+{
+	static constexpr int Version = 1;
+	GEOMETRYFLOW_NODE_INTERNAL(FSourceNodeBase, Version, FNode)
+public:
+	FSourceNodeBase() = default;
+
+	virtual ~FSourceNodeBase() = default;
+
+
+	virtual int GetSourceDataType() const {return -1;}
+
+	// source nodes have no input ( hence no requirements)
+	virtual void CollectRequirements(const TArray<FString>& Outputs, TArray<FEvalRequirement>& RequiredInputsOut) override
+	{
+		return;
+	}
+	// source nodes have no input ( hence no requirements)
+	virtual void CollectAllRequirements(TArray<FEvalRequirement>& RequiredInputsOut) override
+	{
+		return;
+	}
+
+
+	// copy the source data into an object that can be displayed in UI
+	virtual TSharedPtr<FStructOnScope> AsStructOnScope() const = 0;
+
+	// allow the source data to be updated (to reflect changes in the UI)
+	virtual void UpdateSourceFromStructOnScope(const FStructOnScope& StructOnScope) = 0;
+
+};
+
+template<typename T, int32 StorageTypeIdentifier>
+class  TSourceNodeBase : public FSourceNodeBase
+{
+protected:
+	using DataType = TMovableData<T, StorageTypeIdentifier>;
+
+	TSafeSharedPtr<DataType> Value;
+
+public:
+	using CppType = T;
+	static constexpr int DataTypeIdentifier = StorageTypeIdentifier;
+
+	static const FString OutParamValue() { return TEXT("Value"); }
+
+public:
+	TSourceNodeBase()
+	{
+		T InitialValue;
+		Value = MakeSafeShared<DataType>(InitialValue);
+		AddOutput(OutParamValue(), MakeUnique<TBasicNodeOutput<T, StorageTypeIdentifier>>());
+
+		UpdateSourceValue(InitialValue);
+	}
+
+	virtual ~TSourceNodeBase() = default;
+
+	virtual int GetSourceDataType() const override { return DataTypeIdentifier; }
+
+	void UpdateSourceValue(const T& NewValue)
+	{
+		Value->SetData(NewValue);
+		SetOutput(OutParamValue(), Value);
+	}
+	
+	void GetSourceValue(T& OutValue)
+	{
+		Value->GetData(OutValue);
+	}
+
+	virtual void Evaluate(
+		const FNamedDataMap& DatasIn,
+		FNamedDataMap& DatasOut,
+		TUniquePtr<FEvaluationInfo>& EvaluationInfo) override
+	{
+		if (ensure(DatasOut.Contains(OutParamValue())))
+		{
+			DatasOut.SetData(OutParamValue(), GetOutput(OutParamValue()));
+		}
+	}
+
+	// copy the UStruct into an object that can be displayed
+	virtual TSharedPtr<FStructOnScope> AsStructOnScope() const override
+	{
+		TSharedPtr<FStructOnScope> StructOnScope(nullptr);
+		return StructOnScope;
+	}
+
+	// used to update the source UStruct when the displayed version has been changed.
+	virtual void UpdateSourceFromStructOnScope(const FStructOnScope& StructOnScope) override
+	{
+	}
+
+	virtual void Serialize(FArchive& Ar) override
+	{
+		Value->Serialize(Ar);
+	}
+};
+
+
+// Add ability to expose parameters to UI.
+// Note This can only be made for source nodes that have the correct UStructType.
+template<typename UStructType, int32 StorageTypeIdentifier>
+class TUStructSourceNode : public TSourceNodeBase<UStructType, StorageTypeIdentifier>
+{
+	using TSourceNodeBase<UStructType, StorageTypeIdentifier>::Value;
+
+public:
+	using CppType = UStructType;
+
+
+public:
+
+	virtual ~TUStructSourceNode() = default;
+
+	// copy the UStruct into an object that can be displayed
+	virtual TSharedPtr<FStructOnScope> AsStructOnScope() const override
+	{
+		typedef TStructOnScope<UStructType>  FTypedStructOnScope;
+
+		TSharedPtr<FTypedStructOnScope> StructOnScope = MakeShared< FTypedStructOnScope >();
+		StructOnScope->template InitializeAs<UStructType>(Value.Get()->template GetDataConstRef<UStructType>(StorageTypeIdentifier));
+
+		return StructOnScope;
+	}
+
+	// used to update the source UStruct when the displayed version has been changed.
+	virtual void UpdateSourceFromStructOnScope(const FStructOnScope& StructOnScope) override
+	{
+		TStructOnScope<UStructType> TypedStructOnScope;
+		TypedStructOnScope.InitializeFromChecked(StructOnScope);
+
+		UStructType TmpParameters;
+		TmpParameters = *TypedStructOnScope.Get();
+		this->UpdateSourceValue(TmpParameters);
+	}
+
+
+};
+
+namespace
+{
+
+	template <typename PODType>
+	void PODSerializer(FArchive& Ar, PODType& Data)
+	{
+		Ar << Data;
+	}
+
+	template <typename PODStructType>
+	void PODStructSerializer(FArchive& Ar, PODStructType& Data)
+	{
+		int64 NumBits;
+		if (Ar.IsSaving())
+		{
+			NumBits = sizeof(Data)  * 8;
+		}
+		Ar << NumBits;
+		Ar.SerializeBits((void*)&Data, NumBits);
+	}
+
+	template <typename UStructType>
+	void UStructSerializer(FArchive& Ar, UStructType& Data)
+	{
+		UScriptStruct* const Struct = UStructType::StaticStruct(); // singleton
+		Struct->SerializeTaggedProperties(Ar, (uint8*)&Data, Struct, nullptr);
+	}
+}
+
+/**
+* This macro declares a MovableData, Node Input/Output, and TSourceNode type for a given C++ type/identifier
+*
+* eg GEOMETRYFLOW_DECLARE_BASIC_TYPES(Int32, int, (int)EDataTypes::Integer) will declare/define:
+*     - TMovableData type named FDataInt32
+*     - TBasicNodeInput type named FInt32Input
+*     - TBasicNodeOutput type named FInt32Output
+*     - TSourceNode type named FInt32SourceNode
+*
+*/
+#define GEOMETRYFLOW_DECLARE_BASIC_TYPES_WO_SERIALIZATION(TypeName, CppType, TypeIdentifier, VersionIntID)      \
+typedef TMovableData<CppType, TypeIdentifier> FData##TypeName;                   \
+typedef TBasicNodeInput<CppType, TypeIdentifier> F##TypeName##Input;             \
+typedef TBasicNodeOutput<CppType, TypeIdentifier> F##TypeName##Output;           \
+class F##TypeName##SourceNode : public TSourceNodeBase<CppType, TypeIdentifier>  { static constexpr int Version = VersionIntID;  GEOMETRYFLOW_NODE_INTERNAL(F##TypeName##SourceNode, Version, FSourceNodeBase) }; 
+
+/**
+* This macro declares a MovableData, Node Input/Output, and TSourceNode type for a given C++ type/identifier
+* in addition to simple POD serialization for the movable data type
+* 
+* eg GEOMETRYFLOW_DECLARE_BASIC_TYPES(Int32, int, (int)EDataTypes::Integer) will declare/define:
+*     - TMovableData type named FDataInt32
+*     - TBasicNodeInput type named FInt32Input
+*     - TBasicNodeOutput type named FInt32Output
+*     - TSourceNode type named FInt32SourceNode
+*
+*/
+#define GEOMETRYFLOW_DECLARE_BASIC_TYPES(TypeName, CppType, TypeIdentifier, VersionIntID)      \
+template<> struct TSerializationMethod<CppType> { static void Serialize(FArchive& Ar, CppType& Data){ PODStructSerializer(Ar, Data); } }; \
+GEOMETRYFLOW_DECLARE_BASIC_TYPES_WO_SERIALIZATION(TypeName, CppType, TypeIdentifier, VersionIntID) \
+
+/**
+* This macro declares a MovableData, Node Input/Output, and TSourceNode type for a given C++ type/identifier
+* and injects null serialization code.
+*
+* eg GEOMETRYFLOW_DECLARE_BASIC_TYPES(Int32, int, (int)EDataTypes::Integer) will declare/define:
+*     - TMovableData type named FDataInt32
+*     - TBasicNodeInput type named FInt32Input
+*     - TBasicNodeOutput type named FInt32Output
+*     - TSourceNode type named FInt32SourceNode
+*
+*/
+#define GEOMETRYFLOW_DECLARE_BASIC_TYPES_NULL_SERIALIZE(TypeName, CppType, TypeIdentifier, VersionIntID)      \
+template<> struct TSerializationMethod<CppType> { static void Serialize(FArchive& Ar, CppType& Data){} }; \
+GEOMETRYFLOW_DECLARE_BASIC_TYPES_WO_SERIALIZATION(TypeName, CppType, TypeIdentifier, VersionIntID) \
+
+/**
+* This macro declares a MovableData, Node Input/Output, and TSourceNode type for a given C++ type/identifier
+*
+* eg GEOMETRYFLOW_DECLARE_BASIC_TYPES(Int32, FInt32Setting, (int)EDataTypes::Integer) will declare/define:
+*     - TMovableData type named FDataInt32
+*     - TBasicNodeInput type named FInt32Input
+*     - TBasicNodeOutput type named FInt32Output
+*     - TSourceNode type named FInt32SourceNode
+*
+* where FInt32Setting is a UStruct
+*/
+#define GEOMETRYFLOW_DECLARE_BASIC_USTRUCT_TYPES(TypeName, CppType, TypeIdentifier, VersionIntID) \
+template<> struct TSerializationMethod<CppType> { static void Serialize(FArchive& Ar, CppType& Data){ UStructSerializer(Ar, Data); } }; \
+typedef TMovableData<CppType, TypeIdentifier> FData##TypeName##Struct;              \
+typedef TBasicNodeInput<CppType, TypeIdentifier> F##TypeName##Struct##Input;                \
+typedef TBasicNodeOutput<CppType, TypeIdentifier> F##TypeName##Struct##Output;              \
+class F##TypeName##Struct##SourceNode : public TUStructSourceNode<CppType, TypeIdentifier>  { static constexpr int Version = VersionIntID; GEOMETRYFLOW_NODE_INTERNAL(F##TypeName##Struct##SourceNode, Version, FSourceNodeBase) }; \
+
+
+
+/**
+* This macro defines two classes for a given C++ class:
+*   - TMovableData typedef named FData[XYZ]Settings
+*   - TSourceNode typedef named F[XYZ]SettingsSourceNode
+*
+* The assumption is the class has an integer static/constexpr member SettingsType::DataTypeIdentifier that defines the type integer
+*/
+#define GEOMETRYFLOW_DECLARE_SETTINGS_TYPES_WO_SERIALIZATION(CppType, ReadableName, VersionIntID) \
+typedef TMovableData<CppType, CppType::DataTypeIdentifier> FData##ReadableName##Settings; \
+class F##ReadableName##SettingsSourceNode : public TSourceNodeBase<CppType, CppType::DataTypeIdentifier>  { static constexpr int Version = VersionIntID; GEOMETRYFLOW_NODE_INTERNAL(F##ReadableName##SettingsSourceNode, Version, FSourceNodeBase) }; \
+
+
+/**
+* This macro defines two classes for a given C++ class:
+*   - TMovableData typedef named FData[XYZ]Settings
+*   - TSourceNode typedef named F[XYZ]SettingsSourceNode
+*
+* The assumption is the class has an integer static/constexpr member SettingsType::DataTypeIdentifier that defines the type integer
+*/
+#define GEOMETRYFLOW_DECLARE_SETTINGS_TYPES(CppType, ReadableName, VersionIntID) \
+template<> struct TSerializationMethod<CppType> { static void Serialize(FArchive& Ar, CppType& Data){ PODSerializer(Ar, Data); } }; \
+GEOMETRYFLOW_DECLARE_SETTINGS_TYPES_WO_SERIALIZATION(CppType, ReadableName, VersionIntID)
+
+
+/**
+* This macro defines two classes for a given C++ class:
+*   - TMovableData typedef named FData[XYZ]Settings
+*   - TSourceNode typedef named F[XYZ]SettingsSourceNode
+*
+* The assumption is the class has an integer static/constexpr member SettingsType::DataTypeIdentifier that defines the type integer
+*/
+#define GEOMETRYFLOW_DECLARE_SETTINGS_TYPES_NULL_SERIALIZE(CppType, ReadableName, VersionIntID) \
+template<> struct TSerializationMethod<CppType> { static void Serialize(FArchive& Ar, CppType& Data){ } }; \
+GEOMETRYFLOW_DECLARE_SETTINGS_TYPES_WO_SERIALIZATION(CppType, ReadableName, VersionIntID)
+
+#define GEOMETRYFLOW_DECLARE_USTRUCT_SETTINGS_TYPES(CppType, ReadableName, VersionIntID) \
+template<> struct TSerializationMethod<CppType> { static void Serialize(FArchive& Ar, CppType& Data){ UStructSerializer(Ar, Data); } }; \
+typedef TMovableData<CppType, CppType::DataTypeIdentifier> FData##ReadableName##Settings; \
+class F##ReadableName##SettingsSourceNode : public TUStructSourceNode<CppType, CppType::DataTypeIdentifier>{ static constexpr int Version = VersionIntID; GEOMETRYFLOW_NODE_INTERNAL(F##ReadableName##SettingsSourceNode, Version, FSourceNodeBase) }; \
+
+
+
+}	// end namespace GeometryFlow
+}
\ No newline at end of file
diff --git a/Engine/Plugins/Experimental/GeometryFlow/Source/GeometryFlowCore/Public/GeometryFlowTypes.h b/Engine/Plugins/Experimental/GeometryFlow/Source/GeometryFlowCore/Public/GeometryFlowTypes.h
--- a/Engine/Plugins/Experimental/GeometryFlow/Source/GeometryFlowCore/Public/GeometryFlowTypes.h
+++ b/Engine/Plugins/Experimental/GeometryFlow/Source/GeometryFlowCore/Public/GeometryFlowTypes.h
@@ -41,7 +41,8 @@ enum class EGeometryFlowResult
 	NoMatchesFound = 6,
 	MultipleMatchingAmbiguityFound = 7,
 	OperationCancelled = 8,
-	InputAlreadyConnected = 9
+	InputAlreadyConnected = 9,
+	ConnectionRejectedCreatesCycle = 10
 };
 
 
@@ -55,11 +56,20 @@ enum class EDataTypes
 	Float = 1,
 	Double = 2,
 	
+	IntegerStruct = 3,
+	FloatStruct = 4,
+	DoubleStruct = 5,
+
 	Vector3f = 20,
 	Vector3d = 21,
 
+	Vector3fStruct = 22,
+	Vector3dStruct = 23,
+
 	Name = 30,
 
+	NameStruct = 31,
+
 	BaseMeshProcessingTypes = 1000,
 	BaseMeshProcessingTypesEditor = 2000,
 
diff --git a/Engine/Plugins/Experimental/GeometryFlow/Source/GeometryFlowCore/Public/MathNodes/ArithmeticNodes.h b/Engine/Plugins/Experimental/GeometryFlow/Source/GeometryFlowCore/Public/MathNodes/ArithmeticNodes.h
--- a/Engine/Plugins/Experimental/GeometryFlow/Source/GeometryFlowCore/Public/MathNodes/ArithmeticNodes.h
+++ b/Engine/Plugins/Experimental/GeometryFlow/Source/GeometryFlowCore/Public/MathNodes/ArithmeticNodes.h
@@ -3,6 +3,8 @@
 #pragma once
 
 #include "BaseNodes/BinaryOpLambdaNode.h"
+#include "GeometryFlowTypes.h"
+#include "GeometryFlowNode.h"
 
 
 namespace UE
@@ -18,8 +20,11 @@ namespace GeometryFlow
 		{}
 	};
 
-	typedef TBinaryOpAddNode<float, (int)EDataTypes::Float> FAddFloatNode;
-
+	class FAddFloatNode : public TBinaryOpAddNode<float, (int)EDataTypes::Float>
+	{
+		static constexpr int NodeVersion = 1;
+		GEOMETRYFLOW_NODE_INTERNAL(FAddFloatNode, NodeVersion, FNode)
+	};
 
 
 }	// end namespace GeometryFlow
diff --git a/Engine/Plugins/Experimental/GeometryFlow/Source/GeometryFlowMeshProcessing/GeometryFlowMeshProcessing.Build.cs b/Engine/Plugins/Experimental/GeometryFlow/Source/GeometryFlowMeshProcessing/GeometryFlowMeshProcessing.Build.cs
--- a/Engine/Plugins/Experimental/GeometryFlow/Source/GeometryFlowMeshProcessing/GeometryFlowMeshProcessing.Build.cs
+++ b/Engine/Plugins/Experimental/GeometryFlow/Source/GeometryFlowMeshProcessing/GeometryFlowMeshProcessing.Build.cs
@@ -29,7 +29,7 @@ public class GeometryFlowMeshProcessing : ModuleRules
 				"GeometryCore",
 				"GeometryAlgorithms",
 				"DynamicMesh",
-				"GeometryFlowCore"
+				"GeometryFlowCore",
 			}
             );
 			
@@ -37,8 +37,7 @@ public class GeometryFlowMeshProcessing : ModuleRules
 		PrivateDependencyModuleNames.AddRange(
 			new string[]
 			{
-				//"CoreUObject",
-				//"Engine"
+				"CoreUObject",
 			}
 			);
 		
diff --git a/Engine/Plugins/Experimental/GeometryFlow/Source/GeometryFlowMeshProcessing/Private/DataTypes/MeshImageBakingData.cpp b/Engine/Plugins/Experimental/GeometryFlow/Source/GeometryFlowMeshProcessing/Private/DataTypes/MeshImageBakingData.cpp
--- a/Engine/Plugins/Experimental/GeometryFlow/Source/GeometryFlowMeshProcessing/Private/DataTypes/MeshImageBakingData.cpp
+++ b/Engine/Plugins/Experimental/GeometryFlow/Source/GeometryFlowMeshProcessing/Private/DataTypes/MeshImageBakingData.cpp
@@ -38,7 +38,9 @@ void FMakeMeshBakingCacheNode::Evaluate(
 				NewCache->BakeCache.SetDetailMesh(&NewCache->DetailMesh, &NewCache->DetailSpatial);
 				NewCache->BakeCache.SetBakeTargetMesh(&NewCache->TargetMesh);
 
-				NewCache->BakeCache.SetDimensions(Settings.Dimensions);
+				const int32 Width = FMath::Abs(Settings.Dimensions.X);
+				const int32 Height = FMath::Abs(Settings.Dimensions.Y);
+				NewCache->BakeCache.SetDimensions(FImageDimensions(Width, Height));
 				NewCache->BakeCache.SetUVLayer(Settings.UVLayer);
 				NewCache->BakeCache.SetThickness(Settings.Thickness);
 
diff --git a/Engine/Plugins/Experimental/GeometryFlow/Source/GeometryFlowMeshProcessing/Private/GeometryFlowMeshProcessingModule.cpp b/Engine/Plugins/Experimental/GeometryFlow/Source/GeometryFlowMeshProcessing/Private/GeometryFlowMeshProcessingModule.cpp
--- a/Engine/Plugins/Experimental/GeometryFlow/Source/GeometryFlowMeshProcessing/Private/GeometryFlowMeshProcessingModule.cpp
+++ b/Engine/Plugins/Experimental/GeometryFlow/Source/GeometryFlowMeshProcessing/Private/GeometryFlowMeshProcessingModule.cpp
@@ -2,6 +2,7 @@
 
 #include "GeometryFlowMeshProcessingModule.h"
 #include "Modules/ModuleManager.h"
+#include "MeshProcessingNodeRegistration.h"
 
 DEFINE_LOG_CATEGORY(LogGeometryFlowMeshProcessing);
 
@@ -10,6 +11,8 @@ DEFINE_LOG_CATEGORY(LogGeometryFlowMeshProcessing);
 void FGeometryFlowMeshProcessingModule::StartupModule()
 {
 	// This code will execute after your module is loaded into memory; the exact timing is specified in the .uplugin file per-module
+	
+	UE::GeometryFlow::RegisterMeshProcessingNodes();
 }
 
 void FGeometryFlowMeshProcessingModule::ShutdownModule()
diff --git a/Engine/Plugins/Experimental/GeometryFlow/Source/GeometryFlowMeshProcessing/Private/MeshProcessingNodeRegistration.cpp b/Engine/Plugins/Experimental/GeometryFlow/Source/GeometryFlowMeshProcessing/Private/MeshProcessingNodeRegistration.cpp
new file mode 100644
--- /dev/null
+++ b/Engine/Plugins/Experimental/GeometryFlow/Source/GeometryFlowMeshProcessing/Private/MeshProcessingNodeRegistration.cpp
@@ -0,0 +1,148 @@
+// Copyright Epic Games, Inc. All Rights Reserved.
+
+#include "MeshProcessingNodeRegistration.h"
+#include "GeometryFlowNodeFactory.h"
+
+
+#include "DataTypes/CollisionGeometryData.h"
+#include "DataTypes/DynamicMeshData.h"
+#include "DataTypes/IndexSetsData.h"
+#include "DataTypes/MeshImageBakingData.h"
+#include "DataTypes/NormalMapData.h"
+#include "DataTypes/WeightMapData.h"
+#include "DataTypes/TextureImageData.h"
+
+#include "MeshBakingNodes/BakeMeshMultiTextureNode.h"
+#include "MeshBakingNodes/BakeMeshNormalMapNode.h"
+#include "MeshBakingNodes/BakeMeshTextureImageNode.h"
+
+#include "MeshDecompositionNodes/MakeTriangleSetsNode.h"
+
+#include "MeshProcessingNodes/CompactMeshNode.h"
+#include "MeshProcessingNodes/GenerateConvexHullMeshNode.h"
+
+#include "MeshProcessingNodes/MeshDeleteTrianglesNode.h"
+#include "MeshProcessingNodes/MeshMakeCleanGeometryNode.h"
+#include "MeshProcessingNodes/MeshNormalFlowNode.h"
+#include "MeshProcessingNodes/MeshNormalsNodes.h"
+#include "MeshProcessingNodes/MeshRecalculateUVsNode.h"
+
+#include "MeshProcessingNodes/MeshRepackUVsNode.h"
+#include "MeshProcessingNodes/MeshSimplifyNode.h"
+#include "MeshProcessingNodes/MeshSolidifyNode.h"
+
+#include "MeshProcessingNodes/MeshTangentsNodes.h"
+#include "MeshProcessingNodes/MeshThickenNode.h"
+#include "MeshProcessingNodes/MeshVoxMorphologyNode.h"
+#include "MeshProcessingNodes/TransferMeshMaterialIDsNode.h"
+
+
+
+#include "PhysicsNodes/GenerateSimpleCollisionNode.h"
+
+#include "MeshProcessingNodes/MeshProcessingDataTypes.h"
+#include "DynamicMesh/DynamicMesh3.h"
+#include "BaseNodes/SwitchNode.h"
+
+namespace UE
+{
+namespace GeometryFlow
+{
+
+																			  
+void RegisterMeshProcessingNodes()
+{
+
+	FString CategoryName("Data Types");
+	//  DataTypes / 
+	GEOMETRYFLOW_REGISTER_BASIC_TYPES_NODE(CollisionGeometry)
+	GEOMETRYFLOW_REGISTER_NODE_TYPE(CollisionGeometryTransfer)
+
+	GEOMETRYFLOW_REGISTER_BASIC_TYPES_NODE(DynamicMesh)
+	GEOMETRYFLOW_REGISTER_NODE_TYPE(DynamicMeshTransfer)
+
+	GEOMETRYFLOW_REGISTER_BASIC_TYPES_NODE(IndexSets)
+
+	GEOMETRYFLOW_REGISTER_NODE_TYPE(MakeMeshBakingCache)
+	GEOMETRYFLOW_REGISTER_SETTINGS_NODE_TYPE(MeshMakeBakingCache)
+
+	GEOMETRYFLOW_REGISTER_BASIC_TYPES_NODE(NormalMapImage)
+	GEOMETRYFLOW_REGISTER_BASIC_TYPES_NODE(WeightMap)
+
+	GEOMETRYFLOW_REGISTER_BASIC_TYPES_NODE(TextureImage)
+	GEOMETRYFLOW_REGISTER_BASIC_TYPES_NODE(MaterialIDToTextureMap)
+
+
+	// MeshBakingNodes/
+	CategoryName = FString("Baking");
+	GEOMETRYFLOW_REGISTER_NODE_AND_SETTINGS_NODE(BakeMeshMultiTexture)
+	GEOMETRYFLOW_REGISTER_NODE_AND_SETTINGS_NODE(BakeMeshNormalMap)
+	GEOMETRYFLOW_REGISTER_NODE_AND_SETTINGS_NODE(BakeMeshTextureImage)
+
+
+	// MeshDecompositionNodes/
+	CategoryName = FString("Decomposition");
+	GEOMETRYFLOW_REGISTER_NODE_TYPE(MakeTriangleSetsFromMesh)
+	GEOMETRYFLOW_REGISTER_NODE_TYPE(MakeTriangleSetsFromGroups)
+	GEOMETRYFLOW_REGISTER_NODE_TYPE(MakeTriangleSetsFromConnectedComponents)
+	GEOMETRYFLOW_REGISTER_NODE_TYPE(MakeTriangleSetsFromWeightMap)
+
+
+	// MeshProcessingNodes/
+	CategoryName = FString("Processing");
+	GEOMETRYFLOW_REGISTER_NODE_TYPE(CompactMesh)
+	GEOMETRYFLOW_REGISTER_NODE_AND_SETTINGS_NODE(GenerateConvexHullMesh)
+	GEOMETRYFLOW_REGISTER_NODE_TYPE(MeshDeleteTriangles)
+
+	GEOMETRYFLOW_REGISTER_NODE_AND_SETTINGS_NODE(MeshMakeCleanGeometry)
+	GEOMETRYFLOW_REGISTER_NODE_AND_SETTINGS_NODE(MeshNormalFlow)
+
+	GEOMETRYFLOW_REGISTER_SETTINGS_NODE_TYPE(Normals)
+	GEOMETRYFLOW_REGISTER_NODE_TYPE(ComputeMeshNormals)
+	GEOMETRYFLOW_REGISTER_NODE_TYPE(ComputeMeshPerVertexOverlayNormals)
+	GEOMETRYFLOW_REGISTER_NODE_TYPE(ComputeMeshPerVertexNormals)
+
+	GEOMETRYFLOW_REGISTER_NODE_AND_SETTINGS_NODE(MeshRecalculateUVs)
+	GEOMETRYFLOW_REGISTER_NODE_AND_SETTINGS_NODE(MeshRepackUVs)
+
+	GEOMETRYFLOW_REGISTER_SETTINGS_NODE_TYPE(Simplify)
+	GEOMETRYFLOW_REGISTER_NODE_TYPE(SimplifyMesh)
+
+	GEOMETRYFLOW_REGISTER_SETTINGS_NODE_TYPE(Solidify)
+	GEOMETRYFLOW_REGISTER_NODE_TYPE(SolidifyMesh)
+
+	GEOMETRYFLOW_REGISTER_SETTINGS_NODE_TYPE(Tangents)
+	GEOMETRYFLOW_REGISTER_NODE_TYPE(ComputeMeshTangents)
+
+	GEOMETRYFLOW_REGISTER_SETTINGS_NODE_TYPE(Thicken)
+	GEOMETRYFLOW_REGISTER_NODE_TYPE(MeshThicken)
+
+
+	GEOMETRYFLOW_REGISTER_SETTINGS_NODE_TYPE(VoxOffset)
+	GEOMETRYFLOW_REGISTER_SETTINGS_NODE_TYPE(VoxClosure)
+	GEOMETRYFLOW_REGISTER_SETTINGS_NODE_TYPE(VoxOpening)
+
+	GEOMETRYFLOW_REGISTER_NODE_TYPE(VoxDilateMesh)
+	GEOMETRYFLOW_REGISTER_NODE_TYPE(VoxClosureMesh)
+	GEOMETRYFLOW_REGISTER_NODE_TYPE(VoxOpeningMesh)
+
+	GEOMETRYFLOW_REGISTER_NODE_TYPE(TransferMeshMaterialIDs)
+
+
+
+	// PhysicsNodes /
+	CategoryName = FString("Phyics");
+	GEOMETRYFLOW_REGISTER_NODE_AND_SETTINGS_NODE(GenerateSimpleCollision)
+
+
+	CategoryName = FString("Basic Types");
+	typedef TSwitchNode<FDynamicMesh3, 4, (int)EMeshProcessingDataTypes::DynamicMesh> FMeshGeneratorSwitchNode;
+	GEOMETRYFLOW_REGISTER_NODE_TYPE(MeshGeneratorSwitch)
+}
+
+#undef  GEOMETRYFLOW_REGISTER_NODE_TYPE
+#undef  GEOMETRYFLOW_REGISTER_NODE_AND_SETTINGS_NODE
+#undef  GEOMETRYFLOW_QUOTE
+} }  // end namespaces
+
+
diff --git a/Engine/Plugins/Experimental/GeometryFlow/Source/GeometryFlowMeshProcessing/Private/MeshProcessingNodes/MeshNormalFlowNode.cpp b/Engine/Plugins/Experimental/GeometryFlow/Source/GeometryFlowMeshProcessing/Private/MeshProcessingNodes/MeshNormalFlowNode.cpp
--- a/Engine/Plugins/Experimental/GeometryFlow/Source/GeometryFlowMeshProcessing/Private/MeshProcessingNodes/MeshNormalFlowNode.cpp
+++ b/Engine/Plugins/Experimental/GeometryFlow/Source/GeometryFlowMeshProcessing/Private/MeshProcessingNodes/MeshNormalFlowNode.cpp
@@ -49,11 +49,11 @@ void FMeshNormalFlowNode::DoNormalFlow(const FMeshNormalFlowSettings& SettingsIn
 
 	Remesher.bEnableSmoothing = (SettingsIn.SmoothingStrength > 0.0f);
 	Remesher.SmoothSpeedT = SettingsIn.SmoothingStrength;
-	Remesher.SmoothType = SettingsIn.SmoothingType;
+	Remesher.SmoothType = FromUEnum(SettingsIn.SmoothingType);
 
 	if (!SettingsIn.bDiscardAttributes)
 	{
-		Remesher.FlipMetric = (SettingsIn.SmoothingType == FRemesher::ESmoothTypes::Uniform) ?
+		Remesher.FlipMetric = (SettingsIn.SmoothingType == ToUEnum(FRemesher::ESmoothTypes::Uniform)) ?
 			FRemesher::EFlipMetric::OptimalValence : FRemesher::EFlipMetric::MinEdgeLength;
 	}
 
@@ -62,9 +62,9 @@ void FMeshNormalFlowNode::DoNormalFlow(const FMeshNormalFlowSettings& SettingsIn
 
 	FMeshConstraints Constraints;
 	FMeshConstraintsUtil::ConstrainAllBoundariesAndSeams(Constraints, EditMesh,
-														 SettingsIn.MeshBoundaryConstraints,
-														 SettingsIn.GroupBorderConstraints,
-														 SettingsIn.MaterialBorderConstraints,
+														 FromUEnum(SettingsIn.MeshBoundaryConstraints),
+														 FromUEnum(SettingsIn.GroupBorderConstraints),
+														 FromUEnum(SettingsIn.MaterialBorderConstraints),
 														 true, !SettingsIn.bPreserveSharpEdges);
 
 	Remesher.SetExternalConstraints(MoveTemp(Constraints));
diff --git a/Engine/Plugins/Experimental/GeometryFlow/Source/GeometryFlowMeshProcessing/Private/MeshProcessingNodes/MeshRecalculateUVsNode.cpp b/Engine/Plugins/Experimental/GeometryFlow/Source/GeometryFlowMeshProcessing/Private/MeshProcessingNodes/MeshRecalculateUVsNode.cpp
--- a/Engine/Plugins/Experimental/GeometryFlow/Source/GeometryFlowMeshProcessing/Private/MeshProcessingNodes/MeshRecalculateUVsNode.cpp
+++ b/Engine/Plugins/Experimental/GeometryFlow/Source/GeometryFlowMeshProcessing/Private/MeshProcessingNodes/MeshRecalculateUVsNode.cpp
@@ -30,7 +30,7 @@ void FMeshRecalculateUVsNode::RecalculateUVsOnMesh(FDynamicMesh3& EditMesh, cons
 
 		FUVEditResult EditResult;
 
-		ERecalculateUVsUnwrapType UseUnwrapType = Settings.UnwrapType;
+		ERecalculateUVsUnwrapType UseUnwrapType = FromUEnum(Settings.UnwrapType);
 		if (UseUnwrapType == ERecalculateUVsUnwrapType::Auto)
 		{
 			UseUnwrapType = ERecalculateUVsUnwrapType::Conformal;
diff --git a/Engine/Plugins/Experimental/GeometryFlow/Source/GeometryFlowMeshProcessing/Private/PhysicsNodes/GenerateSimpleCollisionNode.cpp b/Engine/Plugins/Experimental/GeometryFlow/Source/GeometryFlowMeshProcessing/Private/PhysicsNodes/GenerateSimpleCollisionNode.cpp
--- a/Engine/Plugins/Experimental/GeometryFlow/Source/GeometryFlowMeshProcessing/Private/PhysicsNodes/GenerateSimpleCollisionNode.cpp
+++ b/Engine/Plugins/Experimental/GeometryFlow/Source/GeometryFlowMeshProcessing/Private/PhysicsNodes/GenerateSimpleCollisionNode.cpp
@@ -107,6 +107,7 @@ namespace GenerateSimpleCollisionNodeHelpers
 		}
 	}
 
+	FMeshSimpleShapeApproximation::EProjectedHullAxisMode FromUEnum(const EGeometryFlow_ProjectedHullAxisMode& Value) { return static_cast<FMeshSimpleShapeApproximation::EProjectedHullAxisMode>(Value); }
 } // namespace GenerateSimpleCollisionNodeHelpers
 
 
@@ -148,35 +149,35 @@ EGeometryFlowResult FGenerateSimpleCollisionNode::EvaluateInternal(const FDynami
 
 	 switch (Settings.Type)
 	 {
-	 case ESimpleCollisionGeometryType::AlignedBoxes:
+	 case EGeometryFlow_SimpleCollisionGeometryType::AlignedBoxes:
 		 ShapeApproximator.bDetectBoxes = true;
 		 ShapeApproximator.Generate_AlignedBoxes(OutCollisionGeometry.Geometry);
 		 break;
-	 case ESimpleCollisionGeometryType::OrientedBoxes:
+	 case EGeometryFlow_SimpleCollisionGeometryType::OrientedBoxes:
 		 ShapeApproximator.bDetectBoxes = true;
 		 ShapeApproximator.Generate_OrientedBoxes(OutCollisionGeometry.Geometry, Progress);
 		 break;
-	 case ESimpleCollisionGeometryType::MinimalSpheres:
+	 case EGeometryFlow_SimpleCollisionGeometryType::MinimalSpheres:
 		 ShapeApproximator.bDetectSpheres = true;
 		 ShapeApproximator.Generate_MinimalSpheres(OutCollisionGeometry.Geometry);
 		 break;
-	 case ESimpleCollisionGeometryType::Capsules:
+	 case EGeometryFlow_SimpleCollisionGeometryType::Capsules:
 		 ShapeApproximator.bDetectCapsules = true;
 		 ShapeApproximator.Generate_Capsules(OutCollisionGeometry.Geometry);
 		 break;
-	 case ESimpleCollisionGeometryType::ConvexHulls:
+	 case EGeometryFlow_SimpleCollisionGeometryType::ConvexHulls:
 		 GenerateSimpleCollisionNodeHelpers::GenerateConvexHulls(Mesh, IndexData, Settings, OutCollisionGeometry, Progress);
 		 break;
-	 case ESimpleCollisionGeometryType::SweptHulls:
+	 case EGeometryFlow_SimpleCollisionGeometryType::SweptHulls:
 		 ShapeApproximator.bDetectConvexes = true;
 		 ShapeApproximator.bSimplifyHulls = Settings.SweptHullSettings.bSimplifyPolygons;
 		 ShapeApproximator.HullSimplifyTolerance = Settings.SweptHullSettings.HullTolerance;
-		 ShapeApproximator.Generate_ProjectedHulls(OutCollisionGeometry.Geometry, Settings.SweptHullSettings.SweepAxis);
+		 ShapeApproximator.Generate_ProjectedHulls(OutCollisionGeometry.Geometry, GenerateSimpleCollisionNodeHelpers::FromUEnum(Settings.SweptHullSettings.SweepAxis));
 		 break;
-	 case ESimpleCollisionGeometryType::MinVolume:
+	 case EGeometryFlow_SimpleCollisionGeometryType::MinVolume:
 		 ShapeApproximator.Generate_MinVolume(OutCollisionGeometry.Geometry);
 		 break;
-	 case ESimpleCollisionGeometryType::None:
+	 case EGeometryFlow_SimpleCollisionGeometryType::None:
 		 break;
 	 }
 
diff --git a/Engine/Plugins/Experimental/GeometryFlow/Source/GeometryFlowMeshProcessing/Public/DataTypes/CollisionGeometryData.h b/Engine/Plugins/Experimental/GeometryFlow/Source/GeometryFlowMeshProcessing/Public/DataTypes/CollisionGeometryData.h
--- a/Engine/Plugins/Experimental/GeometryFlow/Source/GeometryFlowMeshProcessing/Public/DataTypes/CollisionGeometryData.h
+++ b/Engine/Plugins/Experimental/GeometryFlow/Source/GeometryFlowMeshProcessing/Public/DataTypes/CollisionGeometryData.h
@@ -25,7 +25,8 @@ struct FCollisionGeometry
 
 
 // declares FDataCollisionGeometry, FCollisionGeometryInput, FCollisionGeometryOutput, FCollisionGeometrySourceNode
-GEOMETRYFLOW_DECLARE_BASIC_TYPES(CollisionGeometry, FCollisionGeometry, (int)EMeshProcessingDataTypes::CollisionGeometry);
+// @todo - add serialization.  
+GEOMETRYFLOW_DECLARE_BASIC_TYPES_NULL_SERIALIZE(CollisionGeometry, FCollisionGeometry, (int)EMeshProcessingDataTypes::CollisionGeometry, 1);
 
 
 typedef TTransferNode<FCollisionGeometry, (int)EMeshProcessingDataTypes::CollisionGeometry> FCollisionGeometryTransferNode;
diff --git a/Engine/Plugins/Experimental/GeometryFlow/Source/GeometryFlowMeshProcessing/Public/DataTypes/DynamicMeshData.h b/Engine/Plugins/Experimental/GeometryFlow/Source/GeometryFlowMeshProcessing/Public/DataTypes/DynamicMeshData.h
--- a/Engine/Plugins/Experimental/GeometryFlow/Source/GeometryFlowMeshProcessing/Public/DataTypes/DynamicMeshData.h
+++ b/Engine/Plugins/Experimental/GeometryFlow/Source/GeometryFlowMeshProcessing/Public/DataTypes/DynamicMeshData.h
@@ -2,6 +2,7 @@
 
 #pragma once
 
+#include "DynamicMesh/DynamicMesh3.h"
 #include "GeometryFlowCoreNodes.h"
 #include "MeshProcessingNodes/MeshProcessingDataTypes.h"
 
@@ -16,7 +17,16 @@ namespace GeometryFlow
 using namespace UE::Geometry;
 
 // declares FDataDynamicMesh, FDynamicMeshInput, FDynamicMeshOutput, FDynamicMeshSourceNode
-GEOMETRYFLOW_DECLARE_BASIC_TYPES(DynamicMesh, FDynamicMesh3, (int)EMeshProcessingDataTypes::DynamicMesh)
+template<>
+struct TSerializationMethod<FDynamicMesh3> { static void Serialize(FArchive& Ar, FDynamicMesh3& Mesh) { Ar << Mesh; } };
+typedef TMovableData<FDynamicMesh3, (int)EMeshProcessingDataTypes::DynamicMesh> FDataDynamicMesh;
+typedef TBasicNodeInput<FDynamicMesh3, (int)EMeshProcessingDataTypes::DynamicMesh> FDynamicMeshInput;
+typedef TBasicNodeOutput<FDynamicMesh3, (int)EMeshProcessingDataTypes::DynamicMesh> FDynamicMeshOutput;
+class FDynamicMeshSourceNode : public  TSourceNodeBase<FDynamicMesh3, (int)EMeshProcessingDataTypes::DynamicMesh>
+{
+	static constexpr int Version = 1;
+	GEOMETRYFLOW_NODE_INTERNAL(FDynamicMeshSourceNode, Version, FSourceNodeBase)
+};
 
 
 typedef TTransferNode<FDynamicMesh3, (int)EMeshProcessingDataTypes::DynamicMesh> FDynamicMeshTransferNode;
diff --git a/Engine/Plugins/Experimental/GeometryFlow/Source/GeometryFlowMeshProcessing/Public/DataTypes/IndexSetsData.h b/Engine/Plugins/Experimental/GeometryFlow/Source/GeometryFlowMeshProcessing/Public/DataTypes/IndexSetsData.h
--- a/Engine/Plugins/Experimental/GeometryFlow/Source/GeometryFlowMeshProcessing/Public/DataTypes/IndexSetsData.h
+++ b/Engine/Plugins/Experimental/GeometryFlow/Source/GeometryFlowMeshProcessing/Public/DataTypes/IndexSetsData.h
@@ -50,7 +50,9 @@ struct FIndexSets
 
 
 // declares FDataIndexSets, FIndexSetsInput, FIndexSetsOutput, FIndexSetsSourceNode
-GEOMETRYFLOW_DECLARE_BASIC_TYPES(IndexSets, FIndexSets, (int)EMeshProcessingDataTypes::IndexSets)
+template<> struct TSerializationMethod<FIndexSets> { static void Serialize(FArchive& Ar, FIndexSets& Data) { Ar << Data.IndexSets; } }; \
+GEOMETRYFLOW_DECLARE_BASIC_TYPES_WO_SERIALIZATION(IndexSets, FIndexSets, (int)EMeshProcessingDataTypes::IndexSets, 1)
+
 
 
 
diff --git a/Engine/Plugins/Experimental/GeometryFlow/Source/GeometryFlowMeshProcessing/Public/DataTypes/MeshImageBakingData.h b/Engine/Plugins/Experimental/GeometryFlow/Source/GeometryFlowMeshProcessing/Public/DataTypes/MeshImageBakingData.h
--- a/Engine/Plugins/Experimental/GeometryFlow/Source/GeometryFlowMeshProcessing/Public/DataTypes/MeshImageBakingData.h
+++ b/Engine/Plugins/Experimental/GeometryFlow/Source/GeometryFlowMeshProcessing/Public/DataTypes/MeshImageBakingData.h
@@ -8,6 +8,30 @@
 #include "GeometryFlowMovableData.h"
 #include "Sampling/MeshImageBakingCache.h"
 #include "Image/ImageOccupancyMap.h"
+#include "Image/ImageDimensions.h"
+
+
+#include "MeshImageBakingData.generated.h"
+
+USTRUCT()
+struct GEOMETRYFLOWMESHPROCESSING_API FMeshMakeBakingCacheSettings
+{
+
+	GENERATED_USTRUCT_BODY()
+	DECLARE_GEOMETRYFLOW_DATA_TYPE_IDENTIFIER(UE::GeometryFlow::EMeshProcessingDataTypes::MakeBakingCacheSettings);
+
+
+	UPROPERTY(EditAnywhere, Category = "Geometry Flow")
+	FIntPoint Dimensions;  //FImageDimensions Dimensions;
+	
+	UPROPERTY(EditAnywhere, Category = "Geometry Flow")
+	int32 UVLayer = 0;
+
+	UPROPERTY(EditAnywhere, Category = "Geometry Flow")
+	float Thickness = 0.1;
+
+};
+
 
 namespace UE
 {
@@ -16,6 +40,8 @@ namespace GeometryFlow
 
 using namespace UE::Geometry;
 
+typedef FMeshMakeBakingCacheSettings FMeshMakeBakingCacheSettings;
+
 struct GEOMETRYFLOWMESHPROCESSING_API FMeshBakingCache
 {
 	DECLARE_GEOMETRYFLOW_DATA_TYPE_IDENTIFIER(EMeshProcessingDataTypes::BakingCache);
@@ -32,19 +58,9 @@ struct GEOMETRYFLOWMESHPROCESSING_API FMeshBakingCache
 typedef TImmutableData<FMeshBakingCache, FMeshBakingCache::DataTypeIdentifier> FMeshBakingCacheData;
 
 
+GEOMETRYFLOW_DECLARE_USTRUCT_SETTINGS_TYPES(FMeshMakeBakingCacheSettings, MeshMakeBakingCache, 1);
 
 
-struct GEOMETRYFLOWMESHPROCESSING_API FMeshMakeBakingCacheSettings
-{
-	DECLARE_GEOMETRYFLOW_DATA_TYPE_IDENTIFIER(EMeshProcessingDataTypes::MakeBakingCacheSettings);
-
-	FImageDimensions Dimensions;
-	int UVLayer = 0;
-	float Thickness = 0.1;
-
-};
-GEOMETRYFLOW_DECLARE_SETTINGS_TYPES(FMeshMakeBakingCacheSettings, MeshMakeBakingCache);
-
 
 
 
diff --git a/Engine/Plugins/Experimental/GeometryFlow/Source/GeometryFlowMeshProcessing/Public/DataTypes/NormalMapData.h b/Engine/Plugins/Experimental/GeometryFlow/Source/GeometryFlowMeshProcessing/Public/DataTypes/NormalMapData.h
--- a/Engine/Plugins/Experimental/GeometryFlow/Source/GeometryFlowMeshProcessing/Public/DataTypes/NormalMapData.h
+++ b/Engine/Plugins/Experimental/GeometryFlow/Source/GeometryFlowMeshProcessing/Public/DataTypes/NormalMapData.h
@@ -24,8 +24,8 @@ struct FNormalMapImage
 
 
 // declares FDataNormalMapImage, FNormalMapImageInput, FNormalMapImageOutput, FNormalMapImageSourceNode
-GEOMETRYFLOW_DECLARE_BASIC_TYPES(NormalMapImage, FNormalMapImage, (int)EMeshProcessingDataTypes::NormalMapImage)
-
+// @todo - serialize, or reconsider this node as just being some form of graph output.
+GEOMETRYFLOW_DECLARE_BASIC_TYPES_NULL_SERIALIZE(NormalMapImage, FNormalMapImage, (int)EMeshProcessingDataTypes::NormalMapImage, 1)
 
 
 
diff --git a/Engine/Plugins/Experimental/GeometryFlow/Source/GeometryFlowMeshProcessing/Public/DataTypes/TextureImageData.h b/Engine/Plugins/Experimental/GeometryFlow/Source/GeometryFlowMeshProcessing/Public/DataTypes/TextureImageData.h
--- a/Engine/Plugins/Experimental/GeometryFlow/Source/GeometryFlowMeshProcessing/Public/DataTypes/TextureImageData.h
+++ b/Engine/Plugins/Experimental/GeometryFlow/Source/GeometryFlowMeshProcessing/Public/DataTypes/TextureImageData.h
@@ -22,7 +22,7 @@ struct FTextureImage
 };
 
 // declares FDataTextureImage, FTextureImageInput, FTextureImageOutput, FTextureImageSourceNode
-GEOMETRYFLOW_DECLARE_BASIC_TYPES(TextureImage, FTextureImage, (int)EMeshProcessingDataTypes::TextureImage)
+GEOMETRYFLOW_DECLARE_BASIC_TYPES_NULL_SERIALIZE(TextureImage, FTextureImage, (int)EMeshProcessingDataTypes::TextureImage, 1)
 
 
 struct FMaterialIDToTextureMap
@@ -32,7 +32,8 @@ struct FMaterialIDToTextureMap
 	TMap<int32, TSafeSharedPtr<TImageBuilder<FVector4f>>> MaterialIDTextureMap;
 };
 
-GEOMETRYFLOW_DECLARE_BASIC_TYPES(MaterialIDToTextureMap, FMaterialIDToTextureMap, (int)EMeshProcessingDataTypes::MaterialIDToTextureMap)
+// @todo add some form of serialization 
+GEOMETRYFLOW_DECLARE_BASIC_TYPES_NULL_SERIALIZE(MaterialIDToTextureMap, FMaterialIDToTextureMap, (int)EMeshProcessingDataTypes::MaterialIDToTextureMap, 1)
 
 
 }	// end namespace GeometryFlow
diff --git a/Engine/Plugins/Experimental/GeometryFlow/Source/GeometryFlowMeshProcessing/Public/DataTypes/WeightMapData.h b/Engine/Plugins/Experimental/GeometryFlow/Source/GeometryFlowMeshProcessing/Public/DataTypes/WeightMapData.h
--- a/Engine/Plugins/Experimental/GeometryFlow/Source/GeometryFlowMeshProcessing/Public/DataTypes/WeightMapData.h
+++ b/Engine/Plugins/Experimental/GeometryFlow/Source/GeometryFlowMeshProcessing/Public/DataTypes/WeightMapData.h
@@ -22,7 +22,10 @@ struct FWeightMap
 };
 
 
-GEOMETRYFLOW_DECLARE_BASIC_TYPES(WeightMap, FWeightMap, (int)EMeshProcessingDataTypes::WeightMap)
+// @todo - serialize, or reconsider this node as just being some form of graph output.
+template <> struct TSerializationMethod<FWeightMap> { static void Serialize(FArchive& Ar, FWeightMap& Data) { Ar << Data.Weights; } }; \
+GEOMETRYFLOW_DECLARE_BASIC_TYPES_WO_SERIALIZATION(WeightMap, FWeightMap, (int)EMeshProcessingDataTypes::WeightMap, 1)
+
 
 
 }	// end namespace GeometryFlow
diff --git a/Engine/Plugins/Experimental/GeometryFlow/Source/GeometryFlowMeshProcessing/Public/MeshBakingNodes/BakeMeshMultiTextureNode.h b/Engine/Plugins/Experimental/GeometryFlow/Source/GeometryFlowMeshProcessing/Public/MeshBakingNodes/BakeMeshMultiTextureNode.h
--- a/Engine/Plugins/Experimental/GeometryFlow/Source/GeometryFlowMeshProcessing/Public/MeshBakingNodes/BakeMeshMultiTextureNode.h
+++ b/Engine/Plugins/Experimental/GeometryFlow/Source/GeometryFlowMeshProcessing/Public/MeshBakingNodes/BakeMeshMultiTextureNode.h
@@ -2,25 +2,34 @@
 
 #pragma once
 
+#include "GeometryFlowCoreNodes.h"
 #include "MeshBakingNodes/BakeMeshTextureImageNode.h"
 
+#include "BakeMeshMultiTextureNode.generated.h"
+
 class UTexture2D;
 
+USTRUCT()
+struct GEOMETRYFLOWMESHPROCESSING_API FBakeMeshMultiTextureSettings : public FBakeMeshTextureImageSettings
+{
+	GENERATED_USTRUCT_BODY()
+	DECLARE_GEOMETRYFLOW_DATA_TYPE_IDENTIFIER(UE::GeometryFlow::EMeshProcessingDataTypes::BakeMultiTextureSettings);
+};
+
+
 namespace UE
 {
 namespace GeometryFlow
 {
 
-struct GEOMETRYFLOWMESHPROCESSING_API FBakeMeshMultiTextureSettings : public FBakeMeshTextureImageSettings
-{
-	DECLARE_GEOMETRYFLOW_DATA_TYPE_IDENTIFIER(EMeshProcessingDataTypes::BakeMultiTextureSettings);
-};
-
-GEOMETRYFLOW_DECLARE_SETTINGS_TYPES(FBakeMeshMultiTextureSettings, BakeMeshMultiTexture);
+GEOMETRYFLOW_DECLARE_USTRUCT_SETTINGS_TYPES(FBakeMeshMultiTextureSettings, BakeMeshMultiTexture, 1);
 
 
 class GEOMETRYFLOWMESHPROCESSING_API FBakeMeshMultiTextureNode : public FNode
 {
+	static constexpr int Version = 1;
+	GEOMETRYFLOW_NODE_INTERNAL(FBakeMeshMultiTextureNode, Version, FNode)
+
 protected:
 	using SettingsDataType = TMovableData<FBakeMeshMultiTextureSettings, FBakeMeshMultiTextureSettings::DataTypeIdentifier>;
 
diff --git a/Engine/Plugins/Experimental/GeometryFlow/Source/GeometryFlowMeshProcessing/Public/MeshBakingNodes/BakeMeshNormalMapNode.h b/Engine/Plugins/Experimental/GeometryFlow/Source/GeometryFlowMeshProcessing/Public/MeshBakingNodes/BakeMeshNormalMapNode.h
--- a/Engine/Plugins/Experimental/GeometryFlow/Source/GeometryFlowMeshProcessing/Public/MeshBakingNodes/BakeMeshNormalMapNode.h
+++ b/Engine/Plugins/Experimental/GeometryFlow/Source/GeometryFlowMeshProcessing/Public/MeshBakingNodes/BakeMeshNormalMapNode.h
@@ -2,10 +2,22 @@
 
 #pragma once
 
+#include "GeometryFlowCoreNodes.h"
 #include "DataTypes/MeshImageBakingData.h"
 #include "DataTypes/NormalMapData.h"
 #include "DynamicMesh/MeshTangents.h"
 
+#include "BakeMeshNormalMapNode.generated.h"
+
+USTRUCT()
+struct GEOMETRYFLOWMESHPROCESSING_API FBakeMeshNormalMapSettings 
+{
+	GENERATED_USTRUCT_BODY()
+	DECLARE_GEOMETRYFLOW_DATA_TYPE_IDENTIFIER(UE::GeometryFlow::EMeshProcessingDataTypes::BakeNormalMapSettings);
+
+	UPROPERTY(EditAnywhere, Category = "Geometry Flow")
+	double MaxDistance = 0.0;
+};
 
 namespace UE
 {
@@ -14,18 +26,15 @@ namespace GeometryFlow
 
 using namespace UE::Geometry;
 
-struct GEOMETRYFLOWMESHPROCESSING_API FBakeMeshNormalMapSettings
-{
-	DECLARE_GEOMETRYFLOW_DATA_TYPE_IDENTIFIER(EMeshProcessingDataTypes::BakeNormalMapSettings);
-
-	double MaxDistance = 0.0;
-};
-GEOMETRYFLOW_DECLARE_SETTINGS_TYPES(FBakeMeshNormalMapSettings, BakeMeshNormalMap);
+GEOMETRYFLOW_DECLARE_USTRUCT_SETTINGS_TYPES(FBakeMeshNormalMapSettings, BakeMeshNormalMap, 1);
 
 
 
 class GEOMETRYFLOWMESHPROCESSING_API FBakeMeshNormalMapNode : public FNode
 {
+	static constexpr int Version = 1;
+	GEOMETRYFLOW_NODE_INTERNAL(FBakeMeshNormalMapNode, Version, FNode)
+
 protected:
 	using SettingsDataType = TMovableData<FBakeMeshNormalMapSettings, FBakeMeshNormalMapSettings::DataTypeIdentifier>;
 
diff --git a/Engine/Plugins/Experimental/GeometryFlow/Source/GeometryFlowMeshProcessing/Public/MeshBakingNodes/BakeMeshTextureImageNode.h b/Engine/Plugins/Experimental/GeometryFlow/Source/GeometryFlowMeshProcessing/Public/MeshBakingNodes/BakeMeshTextureImageNode.h
--- a/Engine/Plugins/Experimental/GeometryFlow/Source/GeometryFlowMeshProcessing/Public/MeshBakingNodes/BakeMeshTextureImageNode.h
+++ b/Engine/Plugins/Experimental/GeometryFlow/Source/GeometryFlowMeshProcessing/Public/MeshBakingNodes/BakeMeshTextureImageNode.h
@@ -2,30 +2,40 @@
 
 #pragma once
 
+#include "GeometryFlowCoreNodes.h"
 #include "DataTypes/MeshImageBakingData.h"
 #include "DataTypes/TextureImageData.h"
 
+#include "BakeMeshTextureImageNode.generated.h"
 
-namespace UE
-{
-namespace GeometryFlow
+USTRUCT()
+struct GEOMETRYFLOWMESHPROCESSING_API FBakeMeshTextureImageSettings 
 {
+	GENERATED_USTRUCT_BODY()
+	DECLARE_GEOMETRYFLOW_DATA_TYPE_IDENTIFIER(UE::GeometryFlow::EMeshProcessingDataTypes::BakeTextureImageSettings);
 
-
-struct GEOMETRYFLOWMESHPROCESSING_API FBakeMeshTextureImageSettings
-{
-	DECLARE_GEOMETRYFLOW_DATA_TYPE_IDENTIFIER(EMeshProcessingDataTypes::BakeTextureImageSettings);
-
+	UPROPERTY(EditAnywhere, Category = "Geometry Flow")
 	int32 DetailUVLayer = 0;
 
+	UPROPERTY(EditAnywhere, Category = "Geometry Flow")
 	double MaxDistance = 0.0;
 };
-GEOMETRYFLOW_DECLARE_SETTINGS_TYPES(FBakeMeshTextureImageSettings, BakeMeshTextureImage);
+
+namespace UE
+{
+namespace GeometryFlow
+{
+
+
+GEOMETRYFLOW_DECLARE_USTRUCT_SETTINGS_TYPES(FBakeMeshTextureImageSettings, BakeMeshTextureImage, 1);
 
 
 
 class GEOMETRYFLOWMESHPROCESSING_API FBakeMeshTextureImageNode : public FNode
 {
+	static constexpr int Version = 1;
+	GEOMETRYFLOW_NODE_INTERNAL(FBakeMeshTextureImageNode, Version, FNode)
+
 protected:
 	using SettingsDataType = TMovableData<FBakeMeshTextureImageSettings, FBakeMeshTextureImageSettings::DataTypeIdentifier>;
 
diff --git a/Engine/Plugins/Experimental/GeometryFlow/Source/GeometryFlowMeshProcessing/Public/MeshDecompositionNodes/MakeTriangleSetsNode.h b/Engine/Plugins/Experimental/GeometryFlow/Source/GeometryFlowMeshProcessing/Public/MeshDecompositionNodes/MakeTriangleSetsNode.h
--- a/Engine/Plugins/Experimental/GeometryFlow/Source/GeometryFlowMeshProcessing/Public/MeshDecompositionNodes/MakeTriangleSetsNode.h
+++ b/Engine/Plugins/Experimental/GeometryFlow/Source/GeometryFlowMeshProcessing/Public/MeshDecompositionNodes/MakeTriangleSetsNode.h
@@ -18,6 +18,8 @@ namespace GeometryFlow
 
 class FMakeTriangleSetsFromMeshNode : public FNode
 {
+	static constexpr int Version = 1;
+	GEOMETRYFLOW_NODE_INTERNAL(FMakeTriangleSetsFromMeshNode, Version, FNode)
 public:
 	static const FString InParam() { return TEXT("Mesh"); }
 
@@ -86,6 +88,8 @@ public:
 ///
 class FMakeTriangleSetsFromGroupsNode : public FMakeTriangleSetsFromMeshNode
 {
+	static constexpr int Version = 1;
+	GEOMETRYFLOW_NODE_INTERNAL(FMakeTriangleSetsFromGroupsNode, Version, FMakeTriangleSetsFromMeshNode)
 public:
 	static const FString InParamGroupLayer() { return TEXT("GroupLayerName"); }
 	static const FString InParamIgnoreGroups() { return TEXT("IgnoreGroups"); }
@@ -196,7 +200,8 @@ public:
 
 class FMakeTriangleSetsFromConnectedComponentsNode : public FMakeTriangleSetsFromMeshNode
 {
-
+	static constexpr int Version = 1;
+	GEOMETRYFLOW_NODE_INTERNAL(FMakeTriangleSetsFromConnectedComponentsNode, Version, FMakeTriangleSetsFromMeshNode)
 public:
 
 	virtual void ComputeIndexSets(const FNamedDataMap& DatasIn, const FDynamicMesh3& Mesh, FIndexSets& SetsOut) override
@@ -218,7 +223,8 @@ public:
 /// vertex weight.
 class FMakeTriangleSetsFromWeightMapNode : public FMakeTriangleSetsFromMeshNode
 {
-
+	static constexpr int Version = 1;
+	GEOMETRYFLOW_NODE_INTERNAL(FMakeTriangleSetsFromWeightMapNode, Version, FMakeTriangleSetsFromMeshNode)
 public:
 
 	static const FString InParamWeightMap() { return TEXT("WeightMap"); }
diff --git a/Engine/Plugins/Experimental/GeometryFlow/Source/GeometryFlowMeshProcessing/Public/MeshProcessingNodeRegistration.h b/Engine/Plugins/Experimental/GeometryFlow/Source/GeometryFlowMeshProcessing/Public/MeshProcessingNodeRegistration.h
new file mode 100644
--- /dev/null
+++ b/Engine/Plugins/Experimental/GeometryFlow/Source/GeometryFlowMeshProcessing/Public/MeshProcessingNodeRegistration.h
@@ -0,0 +1,14 @@
+// Copyright Epic Games, Inc. All Rights Reserved.
+
+#pragma once
+
+namespace UE
+{
+namespace GeometryFlow
+{
+
+	// Required registration of the geometry flow graph nodes in the MeshProcessing Module.
+	void RegisterMeshProcessingNodes();
+
+}
+}
\ No newline at end of file
diff --git a/Engine/Plugins/Experimental/GeometryFlow/Source/GeometryFlowMeshProcessing/Public/MeshProcessingNodes/CompactMeshNode.h b/Engine/Plugins/Experimental/GeometryFlow/Source/GeometryFlowMeshProcessing/Public/MeshProcessingNodes/CompactMeshNode.h
--- a/Engine/Plugins/Experimental/GeometryFlow/Source/GeometryFlowMeshProcessing/Public/MeshProcessingNodes/CompactMeshNode.h
+++ b/Engine/Plugins/Experimental/GeometryFlow/Source/GeometryFlowMeshProcessing/Public/MeshProcessingNodes/CompactMeshNode.h
@@ -15,6 +15,9 @@ namespace GeometryFlow
  */
 class FCompactMeshNode : public FProcessMeshBaseNode
 {
+	static constexpr int Version = 1;
+	GEOMETRYFLOW_NODE_INTERNAL(FCompactMeshNode, Version, FProcessMeshBaseNode)
+
 public:
 	FCompactMeshNode()
 	{
diff --git a/Engine/Plugins/Experimental/GeometryFlow/Source/GeometryFlowMeshProcessing/Public/MeshProcessingNodes/GenerateConvexHullMeshNode.h b/Engine/Plugins/Experimental/GeometryFlow/Source/GeometryFlowMeshProcessing/Public/MeshProcessingNodes/GenerateConvexHullMeshNode.h
--- a/Engine/Plugins/Experimental/GeometryFlow/Source/GeometryFlowMeshProcessing/Public/MeshProcessingNodes/GenerateConvexHullMeshNode.h
+++ b/Engine/Plugins/Experimental/GeometryFlow/Source/GeometryFlowMeshProcessing/Public/MeshProcessingNodes/GenerateConvexHullMeshNode.h
@@ -4,6 +4,21 @@
 
 #include "MeshProcessingBaseNodes.h"
 
+#include "GenerateConvexHullMeshNode.generated.h"
+
+USTRUCT()
+struct GEOMETRYFLOWMESHPROCESSING_API FGenerateConvexHullMeshSettings
+{
+	GENERATED_USTRUCT_BODY()
+	DECLARE_GEOMETRYFLOW_DATA_TYPE_IDENTIFIER(UE::GeometryFlow::EMeshProcessingDataTypes::GenerateConvexHullMeshSettings);
+
+	UPROPERTY(EditAnywhere, Category = "Geometry Flow")
+	bool bPrefilterVertices = true;
+
+	UPROPERTY(EditAnywhere, Category = "Geometry Flow")
+	int PrefilterGridResolution = 10;
+};
+
 namespace UE
 {
 namespace GeometryFlow
@@ -11,18 +26,16 @@ namespace GeometryFlow
 
 using namespace UE::Geometry;
 
-struct GEOMETRYFLOWMESHPROCESSING_API FGenerateConvexHullMeshSettings
-{
-	DECLARE_GEOMETRYFLOW_DATA_TYPE_IDENTIFIER(EMeshProcessingDataTypes::GenerateConvexHullMeshSettings);
-	bool bPrefilterVertices = true;
-	int PrefilterGridResolution = 10;
-};
+// bring settings into the namespace.  @todo - fix code the expects the UE::GeometryFlow:: for this struct and delete this
+typedef FGenerateConvexHullMeshSettings FGenerateConvexHullMeshSettings;
 
-GEOMETRYFLOW_DECLARE_SETTINGS_TYPES(FGenerateConvexHullMeshSettings, GenerateConvexHullMesh);
+GEOMETRYFLOW_DECLARE_USTRUCT_SETTINGS_TYPES(FGenerateConvexHullMeshSettings, GenerateConvexHullMesh, 1);
 
 
 class GEOMETRYFLOWMESHPROCESSING_API FGenerateConvexHullMeshNode : public TProcessMeshWithSettingsBaseNode<FGenerateConvexHullMeshSettings>
 {
+	static constexpr int Version = 1;
+	GEOMETRYFLOW_NODE_INTERNAL(FGenerateConvexHullMeshNode, Version, FNode)
 
 public:
 
diff --git a/Engine/Plugins/Experimental/GeometryFlow/Source/GeometryFlowMeshProcessing/Public/MeshProcessingNodes/MeshDeleteTrianglesNode.h b/Engine/Plugins/Experimental/GeometryFlow/Source/GeometryFlowMeshProcessing/Public/MeshProcessingNodes/MeshDeleteTrianglesNode.h
--- a/Engine/Plugins/Experimental/GeometryFlow/Source/GeometryFlowMeshProcessing/Public/MeshProcessingNodes/MeshDeleteTrianglesNode.h
+++ b/Engine/Plugins/Experimental/GeometryFlow/Source/GeometryFlowMeshProcessing/Public/MeshProcessingNodes/MeshDeleteTrianglesNode.h
@@ -15,6 +15,9 @@ namespace GeometryFlow
 
 class FMeshDeleteTrianglesNode : public FProcessMeshBaseNode
 {
+	static constexpr int Version = 1;
+	GEOMETRYFLOW_NODE_INTERNAL(FMeshDeleteTrianglesNode, Version, FProcessMeshBaseNode)
+
 public:
 	static const FString InParamTriangleLists() { return TEXT("TriangleLists"); }
 
diff --git a/Engine/Plugins/Experimental/GeometryFlow/Source/GeometryFlowMeshProcessing/Public/MeshProcessingNodes/MeshMakeCleanGeometryNode.h b/Engine/Plugins/Experimental/GeometryFlow/Source/GeometryFlowMeshProcessing/Public/MeshProcessingNodes/MeshMakeCleanGeometryNode.h
--- a/Engine/Plugins/Experimental/GeometryFlow/Source/GeometryFlowMeshProcessing/Public/MeshProcessingNodes/MeshMakeCleanGeometryNode.h
+++ b/Engine/Plugins/Experimental/GeometryFlow/Source/GeometryFlowMeshProcessing/Public/MeshProcessingNodes/MeshMakeCleanGeometryNode.h
@@ -4,33 +4,64 @@
 
 #include "MeshProcessingNodes/MeshProcessingBaseNodes.h"
 
-namespace UE
-{
-namespace GeometryFlow
-{
+#include "MeshMakeCleanGeometryNode.generated.h"
 
-struct GEOMETRYFLOWMESHPROCESSING_API FMeshMakeCleanGeometrySettings
+USTRUCT()
+struct  FMeshMakeCleanGeometrySettings 
 {
-	DECLARE_GEOMETRYFLOW_DATA_TYPE_IDENTIFIER(EMeshProcessingDataTypes::MakeCleanGeometrySettings);
+	GENERATED_USTRUCT_BODY()
+	DECLARE_GEOMETRYFLOW_DATA_TYPE_IDENTIFIER(UE::GeometryFlow::EMeshProcessingDataTypes::MakeCleanGeometrySettings);
 
+	UPROPERTY(EditAnywhere, Category = "Geometry Flow")
 	int FillHolesEdgeCountThresh = 8;
+
+	UPROPERTY(EditAnywhere, Category = "Geometry Flow")
 	double FillHolesEstimatedAreaFraction = 0.001;
 
+	UPROPERTY(EditAnywhere, Category = "Geometry Flow")
 	bool bDiscardAllAttributes = false;
+
+	UPROPERTY(EditAnywhere, Category = "Geometry Flow")
 	bool bClearUVs = true;
+	
+	UPROPERTY(EditAnywhere, Category = "Geometry Flow")
 	bool bClearNormals = true;
+
+	UPROPERTY(EditAnywhere, Category = "Geometry Flow")
 	bool bClearTangents = true;
+
+	UPROPERTY(EditAnywhere, Category = "Geometry Flow")
 	bool bClearVertexColors = true;
+
+	UPROPERTY(EditAnywhere, Category = "Geometry Flow")
 	bool bClearMaterialIDs = false;
 
+	UPROPERTY(EditAnywhere, Category = "Geometry Flow")
 	bool bOutputMeshVertexNormals = true;
+
+	UPROPERTY(EditAnywhere, Category = "Geometry Flow")
 	bool bOutputOverlayVertexNormals = true;
 };
-GEOMETRYFLOW_DECLARE_SETTINGS_TYPES(FMeshMakeCleanGeometrySettings, MeshMakeCleanGeometry);
+
+namespace UE
+{
+namespace GeometryFlow
+{
+
+// bringing this into the namespace so other code that relies on UE::GeometryFlow::FMeshMakeCleanGeometrySettings will just work
+// @todo, update client code and remove this type def
+typedef FMeshMakeCleanGeometrySettings FMeshMakeCleanGeometrySettings;
+
+
+
+GEOMETRYFLOW_DECLARE_USTRUCT_SETTINGS_TYPES(FMeshMakeCleanGeometrySettings, MeshMakeCleanGeometry, 1);
 
 
 class GEOMETRYFLOWMESHPROCESSING_API FMeshMakeCleanGeometryNode : public TProcessMeshWithSettingsBaseNode<FMeshMakeCleanGeometrySettings>
 {
+	static constexpr int Version = 1;
+	GEOMETRYFLOW_NODE_INTERNAL(FMeshMakeCleanGeometryNode, Version, FNode)
+
 public:
 
 	void ProcessMesh(
diff --git a/Engine/Plugins/Experimental/GeometryFlow/Source/GeometryFlowMeshProcessing/Public/MeshProcessingNodes/MeshNormalFlowNode.h b/Engine/Plugins/Experimental/GeometryFlow/Source/GeometryFlowMeshProcessing/Public/MeshProcessingNodes/MeshNormalFlowNode.h
--- a/Engine/Plugins/Experimental/GeometryFlow/Source/GeometryFlowMeshProcessing/Public/MeshProcessingNodes/MeshNormalFlowNode.h
+++ b/Engine/Plugins/Experimental/GeometryFlow/Source/GeometryFlowMeshProcessing/Public/MeshProcessingNodes/MeshNormalFlowNode.h
@@ -4,32 +4,72 @@
 
 #include "MeshProcessingNodes/MeshSimplifyNode.h"
 #include "Remesher.h"
+#include "GeometryFlowTSourceNode.h"
 
+#include "MeshNormalFlowNode.generated.h"
+
+UENUM()
+enum class EGeometryFlow_SmoothTypes
+{
+	Uniform = 0,	/** Uniform weights, produces regular mesh and fastest convergence */
+	Cotan = 1,		/** Cotangent weights prevent tangential flow and hence preserve triangle shape / texture coordinates, but can become unstable... */
+	MeanValue = 2   /** Mean Value weights also have reduced tangential flow but are never negative and hence more stable */
+};
+
+
+// conversions
 namespace UE
 {
 namespace GeometryFlow
 {
+	static UE::Geometry::FRemesher::ESmoothTypes FromUEnum(const EGeometryFlow_SmoothTypes& SmoothType) { return static_cast<FRemesher::ESmoothTypes>(SmoothType); }
+	static EGeometryFlow_SmoothTypes ToUEnum(const UE::Geometry::FRemesher::ESmoothTypes& SmoothType) { return static_cast<EGeometryFlow_SmoothTypes>(SmoothType); }
+}
+}
 
-using namespace UE::Geometry;
-
-struct FMeshNormalFlowSettings : public FMeshSimplifySettings
+USTRUCT()
+struct GEOMETRYFLOWMESHPROCESSING_API FMeshNormalFlowSettings : public FMeshSimplifySettings
 {
-	DECLARE_GEOMETRYFLOW_DATA_TYPE_IDENTIFIER(EMeshProcessingDataTypes::NormalFlowSettings);
+	GENERATED_USTRUCT_BODY()
+	DECLARE_GEOMETRYFLOW_DATA_TYPE_IDENTIFIER(UE::GeometryFlow::EMeshProcessingDataTypes::NormalFlowSettings);
 
+	UPROPERTY(EditAnywhere, Category = "Geometry Flow")
 	int MaxRemeshIterations = 20;
+
+	UPROPERTY(EditAnywhere, Category = "Geometry Flow")
 	int NumExtraProjectionIterations = 5;
+
+	UPROPERTY(EditAnywhere, Category = "Geometry Flow")
 	bool bFlips = true;
+
+	UPROPERTY(EditAnywhere, Category = "Geometry Flow")
 	bool bSplits = true;
+
+	UPROPERTY(EditAnywhere, Category = "Geometry Flow")
 	bool bCollapses = true;
-	FRemesher::ESmoothTypes SmoothingType = FRemesher::ESmoothTypes::Uniform;
+
+	UPROPERTY(EditAnywhere, Category = "Geometry Flow")
+	EGeometryFlow_SmoothTypes SmoothingType = UE::GeometryFlow::ToUEnum(UE::Geometry::FRemesher::ESmoothTypes::Uniform);
+	
+	UPROPERTY(EditAnywhere, Category = "Geometry Flow")
 	float SmoothingStrength = 0.25f;
 };
 
-GEOMETRYFLOW_DECLARE_SETTINGS_TYPES(FMeshNormalFlowSettings, MeshNormalFlow);
+namespace UE
+{
+namespace GeometryFlow
+{
+
+using namespace UE::Geometry;
 
+GEOMETRYFLOW_DECLARE_USTRUCT_SETTINGS_TYPES(FMeshNormalFlowSettings, MeshNormalFlow, 1);
 
-class FMeshNormalFlowNode : public TProcessMeshWithSettingsBaseNode<FMeshNormalFlowSettings>
+
+class GEOMETRYFLOWMESHPROCESSING_API FMeshNormalFlowNode : public TProcessMeshWithSettingsBaseNode<FMeshNormalFlowSettings>
 {
+	static constexpr int Version = 1;
+	GEOMETRYFLOW_NODE_INTERNAL(FMeshNormalFlowNode, Version, FNode)
+
 public:
 
 	// Beside the intput mesh and settings, we also take a "target mesh" that will be the projection target.
diff --git a/Engine/Plugins/Experimental/GeometryFlow/Source/GeometryFlowMeshProcessing/Public/MeshProcessingNodes/MeshNormalsNodes.h b/Engine/Plugins/Experimental/GeometryFlow/Source/GeometryFlowMeshProcessing/Public/MeshProcessingNodes/MeshNormalsNodes.h
--- a/Engine/Plugins/Experimental/GeometryFlow/Source/GeometryFlowMeshProcessing/Public/MeshProcessingNodes/MeshNormalsNodes.h
+++ b/Engine/Plugins/Experimental/GeometryFlow/Source/GeometryFlowMeshProcessing/Public/MeshProcessingNodes/MeshNormalsNodes.h
@@ -7,6 +7,39 @@
 
 #include "DynamicMesh/MeshNormals.h"
 
+#include "MeshNormalsNodes.generated.h"
+
+UENUM()
+enum class EGeometryFlow_ComputeNormalsType
+{
+	PerTriangle = 0,
+	PerVertex = 1,
+	RecomputeExistingTopology = 2,
+	FromFaceAngleThreshold = 3,
+	FromGroups = 4
+};
+
+USTRUCT()
+struct FMeshNormalsSettings 
+{
+	GENERATED_USTRUCT_BODY()
+	DECLARE_GEOMETRYFLOW_DATA_TYPE_IDENTIFIER(UE::GeometryFlow::EMeshProcessingDataTypes::NormalsSettings);
+
+	UPROPERTY(EditAnywhere, Category = "Geometry Flow")
+	EGeometryFlow_ComputeNormalsType NormalsType = EGeometryFlow_ComputeNormalsType::FromFaceAngleThreshold;
+
+	UPROPERTY(EditAnywhere, Category = "Geometry Flow")
+	bool bInvert = false;
+
+	UPROPERTY(EditAnywhere, Category = "Geometry Flow")
+	bool bAreaWeighted = true;
+
+	UPROPERTY(EditAnywhere, Category = "Geometry Flow")
+	bool bAngleWeighted = true;
+
+	UPROPERTY(EditAnywhere, Category = "Geometry Flow") // for FromAngleThreshold type
+	double AngleThresholdDeg = 180.0;
+};
 
 namespace UE
 {
@@ -24,21 +57,20 @@ enum class EComputeNormalsType
 	FromGroups = 4
 };
 
-struct FMeshNormalsSettings
+static EComputeNormalsType FromUEnum(const EGeometryFlow_ComputeNormalsType& ComputeNormalsType)
 {
-	DECLARE_GEOMETRYFLOW_DATA_TYPE_IDENTIFIER(EMeshProcessingDataTypes::NormalsSettings);
-
-	EComputeNormalsType NormalsType = EComputeNormalsType::FromFaceAngleThreshold;
+	return static_cast<EComputeNormalsType>(ComputeNormalsType);
+}
+static EGeometryFlow_ComputeNormalsType ToUEnum(const EComputeNormalsType& ComputeNormalsType)
+{
+	return static_cast<EGeometryFlow_ComputeNormalsType>(ComputeNormalsType);
+}
 
-	bool bInvert = false;
 
-	bool bAreaWeighted = true;
-	bool bAngleWeighted = true;
+// bring settings into the namespace.  @todo - fix code the expects the UE::GeometryFlow:: for this struct and delete this
+typedef FMeshNormalsSettings FMeshNormalsSettings;
 
-	// for FromAngleThreshold type
-	double AngleThresholdDeg = 180.0;
-};
-GEOMETRYFLOW_DECLARE_SETTINGS_TYPES(FMeshNormalsSettings, Normals);
+GEOMETRYFLOW_DECLARE_USTRUCT_SETTINGS_TYPES(FMeshNormalsSettings, Normals, 1);
 
 
 /**
@@ -46,6 +78,8 @@ GEOMETRYFLOW_DECLARE_SETTINGS_TYPES(FMeshNormalsSettings, Normals);
  */
 class FComputeMeshNormalsNode : public TProcessMeshWithSettingsBaseNode<FMeshNormalsSettings>
 {
+	static constexpr int Version = 1;
+	GEOMETRYFLOW_NODE_INTERNAL(FComputeMeshNormalsNode, Version, FNode)
 public:
 	FComputeMeshNormalsNode() : TProcessMeshWithSettingsBaseNode<FMeshNormalsSettings>()
 	{
@@ -88,13 +122,13 @@ public:
 			Normals = MeshInOut.Attributes()->PrimaryNormals();
 		}
 
-		if (Settings.NormalsType == EComputeNormalsType::PerTriangle)
+		if (Settings.NormalsType == ToUEnum( EComputeNormalsType::PerTriangle))
 		{
 			ensure(Settings.bInvert == false);		// not supported
 			FMeshNormals::InitializeMeshToPerTriangleNormals(&MeshInOut);
 			return;
 		}
-		else if (Settings.NormalsType == EComputeNormalsType::PerVertex)
+		else if (Settings.NormalsType == ToUEnum( EComputeNormalsType::PerVertex))
 		{
 			ensure(Settings.bInvert == false);		// not supported
 			FMeshNormals::InitializeOverlayToPerVertexNormals(Normals, false);
@@ -102,15 +136,15 @@ public:
 		}
 
 
-		if (Settings.NormalsType == EComputeNormalsType::FromFaceAngleThreshold)
+		if (Settings.NormalsType == ToUEnum( EComputeNormalsType::FromFaceAngleThreshold))
 		{
 			FMeshNormals::InitializeOverlayTopologyFromOpeningAngle(&MeshInOut, Normals, Settings.AngleThresholdDeg);
 		}
-		else if (Settings.NormalsType == EComputeNormalsType::FromGroups)
+		else if (Settings.NormalsType == ToUEnum( EComputeNormalsType::FromGroups))
 		{
 			FMeshNormals::InitializeOverlayTopologyFromFaceGroups(&MeshInOut, Normals);
 		}
-		else if (Settings.NormalsType != EComputeNormalsType::RecomputeExistingTopology)
+		else if (Settings.NormalsType != ToUEnum( EComputeNormalsType::RecomputeExistingTopology))
 		{
 			ensure(false);
 		}
@@ -129,6 +163,8 @@ public:
  */
 class FComputeMeshPerVertexOverlayNormalsNode : public FSimpleInPlaceProcessMeshBaseNode
 {
+	static constexpr int Version = 1;
+	GEOMETRYFLOW_NODE_INTERNAL(FComputeMeshPerVertexOverlayNormalsNode, Version, FSimpleInPlaceProcessMeshBaseNode)
 public:
 	virtual void ApplyNodeToMesh(FDynamicMesh3& MeshInOut, TUniquePtr<FEvaluationInfo>& EvaluationInfo) override
 	{
@@ -153,6 +189,8 @@ public:
  */
 class FComputeMeshPerVertexNormalsNode : public FSimpleInPlaceProcessMeshBaseNode
 {
+	static constexpr int Version = 1;
+	GEOMETRYFLOW_NODE_INTERNAL(FComputeMeshPerVertexNormalsNode, Version, FSimpleInPlaceProcessMeshBaseNode)
 public:
 	virtual void ApplyNodeToMesh(FDynamicMesh3& MeshInOut, TUniquePtr<FEvaluationInfo>& EvaluationInfo) override
 	{
diff --git a/Engine/Plugins/Experimental/GeometryFlow/Source/GeometryFlowMeshProcessing/Public/MeshProcessingNodes/MeshProcessingBaseNodes.h b/Engine/Plugins/Experimental/GeometryFlow/Source/GeometryFlowMeshProcessing/Public/MeshProcessingNodes/MeshProcessingBaseNodes.h
--- a/Engine/Plugins/Experimental/GeometryFlow/Source/GeometryFlowMeshProcessing/Public/MeshProcessingNodes/MeshProcessingBaseNodes.h
+++ b/Engine/Plugins/Experimental/GeometryFlow/Source/GeometryFlowMeshProcessing/Public/MeshProcessingNodes/MeshProcessingBaseNodes.h
@@ -18,6 +18,8 @@ using namespace UE::Geometry;
 
 class FProcessMeshBaseNode : public FNode
 {
+	static constexpr int Version = 1;
+	GEOMETRYFLOW_NODE_INTERNAL(FProcessMeshBaseNode, Version, FNode)
 public:
 	static const FString InParamMesh() { return TEXT("Mesh"); }
 	static const FString OutParamResultMesh() { return TEXT("ResultMesh"); }
@@ -141,6 +143,8 @@ public:
 */
 class FSimpleInPlaceProcessMeshBaseNode : public FProcessMeshBaseNode
 {
+	static constexpr int Version = 1;
+	GEOMETRYFLOW_NODE_INTERNAL(FSimpleInPlaceProcessMeshBaseNode, Version, FProcessMeshBaseNode)
 public:
 	FSimpleInPlaceProcessMeshBaseNode()
 	{
@@ -184,6 +188,7 @@ inline void LogMeshProcessing(const TCHAR* Identifier, const TCHAR* Message)
 template<typename SettingsType>
 class TProcessMeshWithSettingsBaseNode : public FNode
 {
+	typedef FNode Super;
 public:
 	using SettingsDataType = TMovableData<SettingsType, SettingsType::DataTypeIdentifier>;
 
diff --git a/Engine/Plugins/Experimental/GeometryFlow/Source/GeometryFlowMeshProcessing/Public/MeshProcessingNodes/MeshRecalculateUVsNode.h b/Engine/Plugins/Experimental/GeometryFlow/Source/GeometryFlowMeshProcessing/Public/MeshProcessingNodes/MeshRecalculateUVsNode.h
--- a/Engine/Plugins/Experimental/GeometryFlow/Source/GeometryFlowMeshProcessing/Public/MeshProcessingNodes/MeshRecalculateUVsNode.h
+++ b/Engine/Plugins/Experimental/GeometryFlow/Source/GeometryFlowMeshProcessing/Public/MeshProcessingNodes/MeshRecalculateUVsNode.h
@@ -4,6 +4,28 @@
 
 #include "MeshProcessingNodes/MeshProcessingBaseNodes.h"
 
+#include "MeshRecalculateUVsNode.generated.h"
+
+UENUM()
+enum class EGeometryFlow_RecalculateUVsUnwrapType : uint8
+{
+	Auto = 0,
+	ExpMap = 1,
+	Conformal = 2
+};
+
+USTRUCT()
+struct GEOMETRYFLOWMESHPROCESSING_API FMeshRecalculateUVsSettings 
+{
+	GENERATED_USTRUCT_BODY()
+	DECLARE_GEOMETRYFLOW_DATA_TYPE_IDENTIFIER(UE::GeometryFlow::EMeshProcessingDataTypes::RecalculateUVsSettings);
+
+	UPROPERTY(EditAnywhere, Category = "Geometry Flow")
+	EGeometryFlow_RecalculateUVsUnwrapType UnwrapType = EGeometryFlow_RecalculateUVsUnwrapType::Auto;
+	
+	UPROPERTY(EditAnywhere, Category = "Geometry Flow")
+	int32 UVLayer = 0;
+};
 namespace UE
 {
 namespace GeometryFlow
@@ -16,21 +38,28 @@ enum class ERecalculateUVsUnwrapType : uint8
 	Conformal = 2
 };
 
-
-struct GEOMETRYFLOWMESHPROCESSING_API FMeshRecalculateUVsSettings
+static ERecalculateUVsUnwrapType FromUEnum(const EGeometryFlow_RecalculateUVsUnwrapType& UVsUnwrapType)
 {
-	DECLARE_GEOMETRYFLOW_DATA_TYPE_IDENTIFIER(EMeshProcessingDataTypes::RecalculateUVsSettings);
+	return static_cast<ERecalculateUVsUnwrapType>(UVsUnwrapType);
+}
+static EGeometryFlow_RecalculateUVsUnwrapType ToUEnum(const ERecalculateUVsUnwrapType& UVsUnwrapType)
+{
+	return static_cast<EGeometryFlow_RecalculateUVsUnwrapType>(UVsUnwrapType);
+}
 
-	ERecalculateUVsUnwrapType UnwrapType = ERecalculateUVsUnwrapType::Auto;
-	int32 UVLayer = 0;
 
-};
-GEOMETRYFLOW_DECLARE_SETTINGS_TYPES(FMeshRecalculateUVsSettings, MeshRecalculateUVs);
+// bringing this into the namespace so other code that relies on UE::GeometryFlow::FMeshRecalculateUVsSettings will just work
+// @todo, update client code and remove this type def
+typedef FMeshRecalculateUVsSettings FMeshRecalculateUVsSettings;
+
+GEOMETRYFLOW_DECLARE_USTRUCT_SETTINGS_TYPES(FMeshRecalculateUVsSettings, MeshRecalculateUVs, 1);
 
 
 
 class GEOMETRYFLOWMESHPROCESSING_API FMeshRecalculateUVsNode : public TProcessMeshWithSettingsBaseNode<FMeshRecalculateUVsSettings>
 {
+	static constexpr int Version = 1;
+	GEOMETRYFLOW_NODE_INTERNAL(FMeshRecalculateUVsNode, Version, FNode)
 public:
 	FMeshRecalculateUVsNode() : TProcessMeshWithSettingsBaseNode<FMeshRecalculateUVsSettings>()
 	{
diff --git a/Engine/Plugins/Experimental/GeometryFlow/Source/GeometryFlowMeshProcessing/Public/MeshProcessingNodes/MeshRepackUVsNode.h b/Engine/Plugins/Experimental/GeometryFlow/Source/GeometryFlowMeshProcessing/Public/MeshProcessingNodes/MeshRepackUVsNode.h
--- a/Engine/Plugins/Experimental/GeometryFlow/Source/GeometryFlowMeshProcessing/Public/MeshProcessingNodes/MeshRepackUVsNode.h
+++ b/Engine/Plugins/Experimental/GeometryFlow/Source/GeometryFlowMeshProcessing/Public/MeshProcessingNodes/MeshRepackUVsNode.h
@@ -4,32 +4,45 @@
 
 #include "MeshProcessingNodes/MeshProcessingBaseNodes.h"
 
-namespace UE
-{
-namespace GeometryFlow
-{
-
-
+#include "MeshRepackUVsNode.generated.h"
 
+USTRUCT()
 struct GEOMETRYFLOWMESHPROCESSING_API FMeshRepackUVsSettings
 {
-	DECLARE_GEOMETRYFLOW_DATA_TYPE_IDENTIFIER(EMeshProcessingDataTypes::RepackUVsSettings);
+	GENERATED_USTRUCT_BODY()
+	DECLARE_GEOMETRYFLOW_DATA_TYPE_IDENTIFIER(UE::GeometryFlow::EMeshProcessingDataTypes::RepackUVsSettings);
 
+	UPROPERTY(EditAnywhere, Category = "Geometry Flow")
 	int32 UVLayer = 0;
 
+	UPROPERTY(EditAnywhere, Category = "Geometry Flow")
 	int32 TextureResolution = 512;
+
+	UPROPERTY(EditAnywhere, Category = "Geometry Flow")
 	int32 GutterSize = 1;
+
+	UPROPERTY(EditAnywhere, Category = "Geometry Flow")
 	bool bAllowFlips = false;
 
+	UPROPERTY(EditAnywhere, Category = "Geometry Flow")
 	FVector2f UVScale = FVector2f::One();
+
+	UPROPERTY(EditAnywhere, Category = "Geometry Flow")
 	FVector2f UVTranslation = FVector2f::Zero();
 };
-GEOMETRYFLOW_DECLARE_SETTINGS_TYPES(FMeshRepackUVsSettings, MeshRepackUVs);
+
+namespace UE
+{
+namespace GeometryFlow
+{
 
 
+GEOMETRYFLOW_DECLARE_USTRUCT_SETTINGS_TYPES(FMeshRepackUVsSettings, MeshRepackUVs, 1);
 
 class GEOMETRYFLOWMESHPROCESSING_API FMeshRepackUVsNode : public TProcessMeshWithSettingsBaseNode<FMeshRepackUVsSettings>
 {
+	static constexpr int Version = 1;
+	GEOMETRYFLOW_NODE_INTERNAL(FMeshRepackUVsNode, Version, FNode)
 public:
 	FMeshRepackUVsNode() : TProcessMeshWithSettingsBaseNode<FMeshRepackUVsSettings>()
 	{
diff --git a/Engine/Plugins/Experimental/GeometryFlow/Source/GeometryFlowMeshProcessing/Public/MeshProcessingNodes/MeshSimplifyNode.h b/Engine/Plugins/Experimental/GeometryFlow/Source/GeometryFlowMeshProcessing/Public/MeshProcessingNodes/MeshSimplifyNode.h
--- a/Engine/Plugins/Experimental/GeometryFlow/Source/GeometryFlowMeshProcessing/Public/MeshProcessingNodes/MeshSimplifyNode.h
+++ b/Engine/Plugins/Experimental/GeometryFlow/Source/GeometryFlowMeshProcessing/Public/MeshProcessingNodes/MeshSimplifyNode.h
@@ -8,23 +8,23 @@
 #include "MeshConstraintsUtil.h"
 #include "ProjectionTargets.h"
 #include "Spatial/MeshAABBTree3.h"
+#include "CoreMinimal.h"
+#include "UObject/ObjectMacros.h"
 
+#include "MeshSimplifyNode.generated.h"
 
-namespace UE
-{
-namespace GeometryFlow
-{
 
-using namespace UE::Geometry;
 
-enum class EMeshSimplifyType
+UENUM()
+enum class EGeometryFlow_MeshSimplifyType
 {
 	Standard = 0,
 	VolumePreserving = 1,
 	AttributeAware = 2
 };
 
-enum class EMeshSimplifyTargetType
+UENUM()
+enum class EGeomtryFlow_MeshSimplifyTargetType
 {
 	TriangleCount = 0,
 	VertexCount = 1,
@@ -32,32 +32,127 @@ enum class EMeshSimplifyTargetType
 	GeometricDeviation = 3
 };
 
-struct FMeshSimplifySettings
+UENUM()
+enum class EGeometryFlow_EdgeRefineFlags
+{
+	/** Edge is unconstrained */
+	NoConstraint = 0,
+	/** Edge cannot be flipped */
+	NoFlip = 1,
+	/** Edge cannot be split */
+	NoSplit = 2,
+	/** Edge cannot be collapsed */
+	NoCollapse = 4,
+	/** Edge cannot be flipped, split, or collapsed */
+	FullyConstrained = NoFlip | NoSplit | NoCollapse,
+	/** Edge can only be split */
+	SplitsOnly = NoFlip | NoCollapse,
+	/** Edge can only flip*/
+	FlipOnly = NoSplit | NoCollapse,
+	/** Edge can only collapse*/
+	CollapseOnly = NoFlip | NoSplit
+};
+
+namespace UE {
+namespace GeometryFlow {
+static EEdgeRefineFlags ToEdgeRefineFlags2(const EGeometryFlow_EdgeRefineFlags& Flag)
+{
+	return static_cast<EEdgeRefineFlags>(Flag);
+}
+} }
+
+USTRUCT()
+struct FMeshSimplifySettings 
 {
-	DECLARE_GEOMETRYFLOW_DATA_TYPE_IDENTIFIER(EMeshProcessingDataTypes::SimplifySettings);
+	GENERATED_USTRUCT_BODY()
 
-	EMeshSimplifyType SimplifyType = EMeshSimplifyType::AttributeAware;
-	EMeshSimplifyTargetType TargetType = EMeshSimplifyTargetType::TrianglePercentage;
+	DECLARE_GEOMETRYFLOW_DATA_TYPE_IDENTIFIER(UE::GeometryFlow::EMeshProcessingDataTypes::SimplifySettings);
 
+	UPROPERTY(EditAnywhere, Category = "Geometry Flow")
+	EGeometryFlow_MeshSimplifyType SimplifyType = EGeometryFlow_MeshSimplifyType::AttributeAware;
+	
+	UPROPERTY(EditAnywhere, Category = "Geometry Flow")
+	EGeomtryFlow_MeshSimplifyTargetType TargetType = EGeomtryFlow_MeshSimplifyTargetType::TrianglePercentage;
+
+	UPROPERTY(EditAnywhere, Category = "Geometry Flow")
 	int TargetCount = 100;
+
+	UPROPERTY(EditAnywhere, Category = "Geometry Flow")
 	float TargetFraction = 0.5;
+
+	UPROPERTY(EditAnywhere, Category = "Geometry Flow")
 	float GeometricTolerance = 0.5;
 
+	UPROPERTY(EditAnywhere, Category = "Geometry Flow")
 	bool bDiscardAttributes = false;
+	
+	UPROPERTY(EditAnywhere, Category = "Geometry Flow")
 	bool bPreventNormalFlips = true;
+	
+	UPROPERTY(EditAnywhere, Category = "Geometry Flow")
 	bool bPreserveSharpEdges = false;
+
+	UPROPERTY(EditAnywhere, Category = "Geometry Flow")
 	bool bAllowSeamCollapse = true;
+
+	UPROPERTY(EditAnywhere, Category = "Geometry Flow")
 	bool bAllowSeamSplits = true;
 
-	EEdgeRefineFlags MeshBoundaryConstraints = EEdgeRefineFlags::NoConstraint;
-	EEdgeRefineFlags GroupBorderConstraints = EEdgeRefineFlags::NoConstraint;
-	EEdgeRefineFlags MaterialBorderConstraints = EEdgeRefineFlags::NoConstraint;
+	UPROPERTY(EditAnywhere, Category = "Geometry Flow")
+	EGeometryFlow_EdgeRefineFlags MeshBoundaryConstraints = EGeometryFlow_EdgeRefineFlags::NoConstraint;
+
+	UPROPERTY(EditAnywhere, Category = "Geometry Flow")
+	EGeometryFlow_EdgeRefineFlags GroupBorderConstraints = EGeometryFlow_EdgeRefineFlags::NoConstraint;
+
+	UPROPERTY(EditAnywhere, Category = "Geometry Flow")
+	EGeometryFlow_EdgeRefineFlags MaterialBorderConstraints = EGeometryFlow_EdgeRefineFlags::NoConstraint;
 };
-GEOMETRYFLOW_DECLARE_SETTINGS_TYPES(FMeshSimplifySettings, Simplify);
 
+namespace UE { namespace GeometryFlow {
+GEOMETRYFLOW_DECLARE_USTRUCT_SETTINGS_TYPES(FMeshSimplifySettings, Simplify, 1);
+}}
+
+namespace UE
+{
+namespace GeometryFlow
+{
+
+using namespace UE::Geometry;
+
+// bringing this into the namespace so other code that relies on UE::GeometryFlow::FMeshSimplifySettings will just work
+// @todo, update client code and remove this type def
+typedef FMeshSimplifySettings FMeshSimplifySettings;
+
+enum class EMeshSimplifyType
+{
+	Standard = 0,
+	VolumePreserving = 1,
+	AttributeAware = 2
+};
+
+enum class EMeshSimplifyTargetType
+{
+	TriangleCount = 0,
+	VertexCount = 1,
+	TrianglePercentage = 2,
+	GeometricDeviation = 3
+};
+
+
+
+static EEdgeRefineFlags FromUEnum(const EGeometryFlow_EdgeRefineFlags& Flag) { return static_cast<EEdgeRefineFlags>(Flag); }
+static EGeometryFlow_EdgeRefineFlags ToUEnum(const EEdgeRefineFlags& Flag) {return static_cast<EGeometryFlow_EdgeRefineFlags>(Flag); }
+
+static EMeshSimplifyType FromUEnum(const EGeometryFlow_MeshSimplifyType& SimplifyType) { return static_cast<EMeshSimplifyType>(SimplifyType);}
+static EGeometryFlow_MeshSimplifyType ToUEnum(const EMeshSimplifyType& SimplifyType) { return static_cast<EGeometryFlow_MeshSimplifyType>(SimplifyType); }
+
+static EMeshSimplifyTargetType FromUEnum(const EGeomtryFlow_MeshSimplifyTargetType& TargetType) { return static_cast<EMeshSimplifyTargetType>(TargetType); }
+static EGeomtryFlow_MeshSimplifyTargetType ToUEnum(const EMeshSimplifyTargetType& TargetType) { return static_cast<EGeomtryFlow_MeshSimplifyTargetType>(TargetType); }
 
 class FSimplifyMeshNode : public TProcessMeshWithSettingsBaseNode<FMeshSimplifySettings>
 {
+	static constexpr int Version = 1;
+	GEOMETRYFLOW_NODE_INTERNAL(FSimplifyMeshNode, Version, FNode)
 public:
 	FSimplifyMeshNode() : TProcessMeshWithSettingsBaseNode<FMeshSimplifySettings>()
 	{
@@ -113,26 +208,26 @@ public:
 
 		FMeshConstraints Constraints;
 		FMeshConstraintsUtil::ConstrainAllBoundariesAndSeams(Constraints, *TargetMesh,
-			Settings.MeshBoundaryConstraints,
-			Settings.GroupBorderConstraints,
-			Settings.MaterialBorderConstraints,
+			FromUEnum(Settings.MeshBoundaryConstraints),
+			FromUEnum(Settings.GroupBorderConstraints),
+			FromUEnum(Settings.MaterialBorderConstraints),
 			Settings.bAllowSeamSplits, !Settings.bPreserveSharpEdges, Settings.bAllowSeamCollapse);
 		Simplifier.SetExternalConstraints(MoveTemp(Constraints));
 
-		if (Settings.TargetType == EMeshSimplifyTargetType::TrianglePercentage)
+		if (Settings.TargetType == EGeomtryFlow_MeshSimplifyTargetType::TrianglePercentage)
 		{
 			int32 UseTarget = FMath::Max(4, (int)(Settings.TargetFraction * (double)TargetMesh->TriangleCount()));
 			Simplifier.SimplifyToTriangleCount(UseTarget);
 		}
-		else if (Settings.TargetType == EMeshSimplifyTargetType::TriangleCount)
+		else if (Settings.TargetType == EGeomtryFlow_MeshSimplifyTargetType::TriangleCount)
 		{
 			Simplifier.SimplifyToTriangleCount( FMath::Max(1,Settings.TargetCount) );
 		}
-		else if (Settings.TargetType == EMeshSimplifyTargetType::VertexCount)
+		else if (Settings.TargetType == EGeomtryFlow_MeshSimplifyTargetType::VertexCount)
 		{
 			Simplifier.SimplifyToVertexCount( FMath::Max(3, Settings.TargetCount) );
 		}
-		else if (Settings.TargetType == EMeshSimplifyTargetType::GeometricDeviation)
+		else if (Settings.TargetType == EGeomtryFlow_MeshSimplifyTargetType::GeometricDeviation)
 		{
 			// need to create projection target to measure error
 			FDynamicMesh3 MeshCopy;
@@ -159,15 +254,15 @@ public:
 			MeshInOut.DiscardAttributes();
 		}
 
-		if (Settings.SimplifyType == EMeshSimplifyType::Standard)
+		if (Settings.SimplifyType == EGeometryFlow_MeshSimplifyType::Standard)
 		{
 			DoSimplifyOfType<FQEMSimplification>(Settings, &MeshInOut, EvaluationInfo);
 		}
-		else if (Settings.SimplifyType == EMeshSimplifyType::VolumePreserving)
+		else if (Settings.SimplifyType == EGeometryFlow_MeshSimplifyType::VolumePreserving)
 		{
 			DoSimplifyOfType<FVolPresMeshSimplification>(Settings, &MeshInOut, EvaluationInfo);
 		}
-		else if (Settings.SimplifyType == EMeshSimplifyType::AttributeAware)
+		else if (Settings.SimplifyType == EGeometryFlow_MeshSimplifyType::AttributeAware)
 		{
 			DoSimplifyOfType<FAttrMeshSimplification>(Settings, &MeshInOut, EvaluationInfo);
 		}
diff --git a/Engine/Plugins/Experimental/GeometryFlow/Source/GeometryFlowMeshProcessing/Public/MeshProcessingNodes/MeshSolidifyNode.h b/Engine/Plugins/Experimental/GeometryFlow/Source/GeometryFlowMeshProcessing/Public/MeshProcessingNodes/MeshSolidifyNode.h
--- a/Engine/Plugins/Experimental/GeometryFlow/Source/GeometryFlowMeshProcessing/Public/MeshProcessingNodes/MeshSolidifyNode.h
+++ b/Engine/Plugins/Experimental/GeometryFlow/Source/GeometryFlowMeshProcessing/Public/MeshProcessingNodes/MeshSolidifyNode.h
@@ -9,28 +9,49 @@
 #include "Spatial/FastWinding.h"
 #include "Spatial/MeshAABBTree3.h"
 
+#include "MeshSolidifyNode.generated.h"
 
-namespace UE
-{
-namespace GeometryFlow
-{
-
-using namespace UE::Geometry;
-
+USTRUCT()
 struct FMeshSolidifySettings
 {
-	DECLARE_GEOMETRYFLOW_DATA_TYPE_IDENTIFIER(EMeshProcessingDataTypes::SolidifySettings);
+	GENERATED_USTRUCT_BODY()
+
+	DECLARE_GEOMETRYFLOW_DATA_TYPE_IDENTIFIER(UE::GeometryFlow::EMeshProcessingDataTypes::SolidifySettings);
 
+	UPROPERTY(EditAnywhere, Category = "Geometry Flow")
 	int VoxelResolution = 64;
+
+	UPROPERTY(EditAnywhere, Category = "Geometry Flow")
 	float WindingThreshold = 0.5;
+
+	UPROPERTY(EditAnywhere, Category = "Geometry Flow")
 	int SurfaceConvergeSteps = 5;
+
+	UPROPERTY(EditAnywhere, Category = "Geometry Flow")
 	float ExtendBounds = 2.0f;
 };
-GEOMETRYFLOW_DECLARE_SETTINGS_TYPES(FMeshSolidifySettings, Solidify);
+
+
+
+namespace UE
+{
+namespace GeometryFlow
+{
+
+using namespace UE::Geometry;
+
+
+// bringing this into the namespace so other code that relies on UE::GeometryFlow::FMeshSolidifySettings will just work
+// @todo, update client code and remove this type def
+typedef FMeshSolidifySettings FMeshSolidifySettings;
+
+GEOMETRYFLOW_DECLARE_USTRUCT_SETTINGS_TYPES(FMeshSolidifySettings, Solidify, 1);
 
 
 class FSolidifyMeshNode : public TProcessMeshWithSettingsBaseNode<FMeshSolidifySettings>
 {
+	static constexpr int Version = 1;
+	GEOMETRYFLOW_NODE_INTERNAL(FSolidifyMeshNode, Version, FNode)
 public:
 
 	virtual void ProcessMesh(
diff --git a/Engine/Plugins/Experimental/GeometryFlow/Source/GeometryFlowMeshProcessing/Public/MeshProcessingNodes/MeshTangentsNodes.h b/Engine/Plugins/Experimental/GeometryFlow/Source/GeometryFlowMeshProcessing/Public/MeshProcessingNodes/MeshTangentsNodes.h
--- a/Engine/Plugins/Experimental/GeometryFlow/Source/GeometryFlowMeshProcessing/Public/MeshProcessingNodes/MeshTangentsNodes.h
+++ b/Engine/Plugins/Experimental/GeometryFlow/Source/GeometryFlowMeshProcessing/Public/MeshProcessingNodes/MeshTangentsNodes.h
@@ -3,6 +3,7 @@
 #pragma once
 
 #include "BaseNodes/TransformerWithSettingsNode.h"
+#include "BaseNodes/TransferNode.h"
 #include "DynamicMesh/DynamicMeshAttributeSet.h"
 #include "MeshProcessingNodes/MeshProcessingDataTypes.h"
 
@@ -10,8 +11,29 @@
 #include "GeometryFlowCoreNodes.h"
 #include "GeometryFlowMovableData.h"
 
+#include "MeshTangentsNodes.generated.h"
+
 namespace UE::GeometryFlow { template <typename T, int StorageTypeIdentifier> class TTransferNode; }
 
+UENUM()
+enum class EGeometryFlow_ComputeTangentsType
+{
+	PerTriangle = 0,
+	FastMikkT = 1
+};
+
+USTRUCT()
+struct FMeshTangentsSettings
+{
+	GENERATED_USTRUCT_BODY()
+	DECLARE_GEOMETRYFLOW_DATA_TYPE_IDENTIFIER(UE::GeometryFlow::EMeshProcessingDataTypes::TangentsSettings);
+
+	UPROPERTY(EditAnywhere, Category = "Geometry Flow")
+	EGeometryFlow_ComputeTangentsType TangentsType = EGeometryFlow_ComputeTangentsType::FastMikkT;
+
+	UPROPERTY(EditAnywhere, Category = "Geometry Flow")
+	int UVLayer = 0;
+};
 
 namespace UE
 {
@@ -26,19 +48,18 @@ enum class EComputeTangentsType
 	FastMikkT = 1
 };
 
+static EComputeTangentsType FromUEnum(const EGeometryFlow_ComputeTangentsType& type) { return static_cast<EComputeTangentsType>(type); }
+static EGeometryFlow_ComputeTangentsType ToUEnum(const EComputeTangentsType& type) { return static_cast<EGeometryFlow_ComputeTangentsType>(type); }
 
-struct FMeshTangentsSettings
-{
-	DECLARE_GEOMETRYFLOW_DATA_TYPE_IDENTIFIER(EMeshProcessingDataTypes::TangentsSettings);
+// bringing this into the namespace so other code that relies on UE::GeometryFlow:: will work
+// @todo, update client code and remove this type def
+typedef FMeshTangentsSettings FMeshTangentsSettings;
 
-	EComputeTangentsType TangentsType = EComputeTangentsType::FastMikkT;
 
-	int UVLayer = 0;
-};
-GEOMETRYFLOW_DECLARE_SETTINGS_TYPES(FMeshTangentsSettings, Tangents);
+GEOMETRYFLOW_DECLARE_USTRUCT_SETTINGS_TYPES(FMeshTangentsSettings, Tangents, 1);
 
 
-typedef TTransferNode<FMeshTangentsd, (int)EMeshProcessingDataTypes::MeshTangentSet> FMeshTangentsTransferNode;
+class FMeshTangentsTransferNode : public  TTransferNode<FMeshTangentsd, (int)EMeshProcessingDataTypes::MeshTangentSet> {static constexpr int Version =1; GEOMETRYFLOW_NODE_INTERNAL(FMeshTangentsTransferNode, Version, FNode) };
 
 
 class FComputeMeshTangentsNode : public TTransformerWithSettingsNode<
@@ -46,6 +67,8 @@ class FComputeMeshTangentsNode : public TTransformerWithSettingsNode<
 	FMeshTangentsSettings, (int)EMeshProcessingDataTypes::TangentsSettings,
 	FMeshTangentsd, (int)EMeshProcessingDataTypes::MeshTangentSet>
 {
+	static constexpr int Version = 1;
+	GEOMETRYFLOW_NODE_INTERNAL(FComputeMeshTangentsNode, Version, FNode)
 public:
 	FComputeMeshTangentsNode() : TTransformerWithSettingsNode()
 	{
@@ -68,7 +91,7 @@ public:
 		}
 
 		FComputeTangentsOptions Options;
-		Options.bAveraged = (Settings.TangentsType == EComputeTangentsType::FastMikkT);
+		Options.bAveraged = (Settings.TangentsType == ToUEnum( EComputeTangentsType::FastMikkT) );
 
 		TangentsOut.SetMesh(&MeshIn);
 		TangentsOut.ComputeTriVertexTangents(
diff --git a/Engine/Plugins/Experimental/GeometryFlow/Source/GeometryFlowMeshProcessing/Public/MeshProcessingNodes/MeshThickenNode.h b/Engine/Plugins/Experimental/GeometryFlow/Source/GeometryFlowMeshProcessing/Public/MeshProcessingNodes/MeshThickenNode.h
--- a/Engine/Plugins/Experimental/GeometryFlow/Source/GeometryFlowMeshProcessing/Public/MeshProcessingNodes/MeshThickenNode.h
+++ b/Engine/Plugins/Experimental/GeometryFlow/Source/GeometryFlowMeshProcessing/Public/MeshProcessingNodes/MeshThickenNode.h
@@ -5,22 +5,32 @@
 #include "MeshProcessingNodes/MeshProcessingBaseNodes.h"
 #include "DataTypes/WeightMapData.h"
 
+#include "MeshThickenNode.generated.h"
+
+USTRUCT()
+struct GEOMETRYFLOWMESHPROCESSING_API FMeshThickenSettings 
+{
+	GENERATED_USTRUCT_BODY()
+	DECLARE_GEOMETRYFLOW_DATA_TYPE_IDENTIFIER(UE::GeometryFlow::EMeshProcessingDataTypes::ThickenSettings);
+
+	float ThickenAmount = 1.0f;
+};
+
 namespace UE
 {
 namespace GeometryFlow
 {
 
-struct GEOMETRYFLOWMESHPROCESSING_API FMeshThickenSettings
-{
-	DECLARE_GEOMETRYFLOW_DATA_TYPE_IDENTIFIER(EMeshProcessingDataTypes::ThickenSettings);
+// @todo - remove this when other code is no longer expecting FMeshThickenSettings to be in the UE::GeometryFlow namespace
+typedef FMeshThickenSettings FMeshThickenSettings;
 
-	float ThickenAmount = 1.0f;
-};
-GEOMETRYFLOW_DECLARE_SETTINGS_TYPES(FMeshThickenSettings, Thicken);
+GEOMETRYFLOW_DECLARE_USTRUCT_SETTINGS_TYPES(FMeshThickenSettings, Thicken, 1);
 
 
 class GEOMETRYFLOWMESHPROCESSING_API FMeshThickenNode : public TProcessMeshWithSettingsBaseNode<FMeshThickenSettings>
 {
+	static constexpr int Version = 1;
+	GEOMETRYFLOW_NODE_INTERNAL(FMeshThickenNode, Version, FNode)
 public:
 
 	static const FString InParamWeightMap() { return TEXT("WeightMap"); }
diff --git a/Engine/Plugins/Experimental/GeometryFlow/Source/GeometryFlowMeshProcessing/Public/MeshProcessingNodes/MeshVoxMorphologyNode.h b/Engine/Plugins/Experimental/GeometryFlow/Source/GeometryFlowMeshProcessing/Public/MeshProcessingNodes/MeshVoxMorphologyNode.h
--- a/Engine/Plugins/Experimental/GeometryFlow/Source/GeometryFlowMeshProcessing/Public/MeshProcessingNodes/MeshVoxMorphologyNode.h
+++ b/Engine/Plugins/Experimental/GeometryFlow/Source/GeometryFlowMeshProcessing/Public/MeshProcessingNodes/MeshVoxMorphologyNode.h
@@ -8,6 +8,21 @@
 #include "DynamicMesh/DynamicMeshAABBTree3.h"
 #include "Implicit/Morphology.h"
 
+#include "MeshVoxMorphologyNode.generated.h"
+
+USTRUCT()
+struct FVoxMorphologyOpSettings 
+{
+	GENERATED_USTRUCT_BODY()
+
+	DECLARE_GEOMETRYFLOW_DATA_TYPE_IDENTIFIER(UE::GeometryFlow::EMeshProcessingDataTypes::VoxMorphologyOpSettings);
+
+	UPROPERTY(EditAnywhere, Category = "Geometry Flow")
+	int VoxelResolution = 64;
+
+	UPROPERTY(EditAnywhere, Category = "Geometry Flow")
+	double Distance = 1.0;
+};
 
 namespace UE
 {
@@ -16,19 +31,25 @@ namespace GeometryFlow
 
 using namespace UE::Geometry;
 
-struct FVoxMorphologyOpSettings
-{
-	DECLARE_GEOMETRYFLOW_DATA_TYPE_IDENTIFIER(EMeshProcessingDataTypes::VoxMorphologyOpSettings);
+template<> 
+struct TSerializationMethod<FVoxMorphologyOpSettings> {static void Serialize(FArchive& Ar, FVoxMorphologyOpSettings& Data) { UStructSerializer(Ar, Data); } };
+
 
-	int VoxelResolution = 64;
-	double Distance = 1.0;
-};
 typedef FVoxMorphologyOpSettings FVoxOffsetSettings;
 typedef FVoxMorphologyOpSettings FVoxClosureSettings;
 typedef FVoxMorphologyOpSettings FVoxOpeningSettings;
-GEOMETRYFLOW_DECLARE_SETTINGS_TYPES(FVoxMorphologyOpSettings, VoxOffset);
-GEOMETRYFLOW_DECLARE_SETTINGS_TYPES(FVoxMorphologyOpSettings, VoxClosure);
-GEOMETRYFLOW_DECLARE_SETTINGS_TYPES(FVoxMorphologyOpSettings, VoxOpening);
+
+
+#define GEOMETRYFLOW_DECLARE_MORPHOLOGY_SETTINGS_TYPES(CppType, ReadableName) \
+typedef TMovableData<CppType, CppType::DataTypeIdentifier> FData##ReadableName##Settings; \
+class F##ReadableName##SettingsSourceNode : public TSourceNodeBase<CppType, CppType::DataTypeIdentifier>{ static constexpr int Version = 1; GEOMETRYFLOW_NODE_INTERNAL(F##ReadableName##SettingsSourceNode, Version, FSourceNodeBase) }; \
+
+
+GEOMETRYFLOW_DECLARE_MORPHOLOGY_SETTINGS_TYPES(FVoxOffsetSettings, VoxOffset);
+GEOMETRYFLOW_DECLARE_MORPHOLOGY_SETTINGS_TYPES(FVoxClosureSettings, VoxClosure);
+GEOMETRYFLOW_DECLARE_MORPHOLOGY_SETTINGS_TYPES(FVoxOpeningSettings, VoxOpening);
+
+#undef GEOMETRYFLOW_DECLARE_MORPHOLOGY_SETTINGS_TYPES
 
 
 template<TImplicitMorphology<FDynamicMesh3>::EMorphologyOp MorphologyOp>
@@ -76,9 +97,9 @@ public:
 
 };
 
-typedef TVoxMorphologyMeshNode<TImplicitMorphology<FDynamicMesh3>::EMorphologyOp::Dilate> FVoxDilateMeshNode;
-typedef TVoxMorphologyMeshNode<TImplicitMorphology<FDynamicMesh3>::EMorphologyOp::Close> FVoxClosureMeshNode;
-typedef TVoxMorphologyMeshNode<TImplicitMorphology<FDynamicMesh3>::EMorphologyOp::Open> FVoxOpeningMeshNode;
+class FVoxDilateMeshNode : public TVoxMorphologyMeshNode<TImplicitMorphology<FDynamicMesh3>::EMorphologyOp::Dilate> { static constexpr int Version = 1; GEOMETRYFLOW_NODE_INTERNAL(FVoxDilateMeshNode, Version, FNode) };
+class FVoxClosureMeshNode: public TVoxMorphologyMeshNode<TImplicitMorphology<FDynamicMesh3>::EMorphologyOp::Close> { static constexpr int Version = 1;  GEOMETRYFLOW_NODE_INTERNAL(FVoxClosureMeshNode, Version, FNode) };
+class FVoxOpeningMeshNode: public  TVoxMorphologyMeshNode<TImplicitMorphology<FDynamicMesh3>::EMorphologyOp::Open> { static constexpr int Version = 1;  GEOMETRYFLOW_NODE_INTERNAL(FVoxOpeningMeshNode, Version, FNode) };
 
 
 }	// end namespace GeometryFlow
diff --git a/Engine/Plugins/Experimental/GeometryFlow/Source/GeometryFlowMeshProcessing/Public/MeshProcessingNodes/TransferMeshMaterialIDsNode.h b/Engine/Plugins/Experimental/GeometryFlow/Source/GeometryFlowMeshProcessing/Public/MeshProcessingNodes/TransferMeshMaterialIDsNode.h
--- a/Engine/Plugins/Experimental/GeometryFlow/Source/GeometryFlowMeshProcessing/Public/MeshProcessingNodes/TransferMeshMaterialIDsNode.h
+++ b/Engine/Plugins/Experimental/GeometryFlow/Source/GeometryFlowMeshProcessing/Public/MeshProcessingNodes/TransferMeshMaterialIDsNode.h
@@ -16,6 +16,8 @@ namespace GeometryFlow
  */
 class FTransferMeshMaterialIDsNode : public FProcessMeshBaseNode
 {
+	static constexpr int Version = 1;
+	GEOMETRYFLOW_NODE_INTERNAL(FTransferMeshMaterialIDsNode, Version, FProcessMeshBaseNode)
 public:
 	static const FString InParamMaterialSourceMesh() { return TEXT("MaterialSourceMesh"); }
 
diff --git a/Engine/Plugins/Experimental/GeometryFlow/Source/GeometryFlowMeshProcessing/Public/PhysicsNodes/GenerateSimpleCollisionNode.h b/Engine/Plugins/Experimental/GeometryFlow/Source/GeometryFlowMeshProcessing/Public/PhysicsNodes/GenerateSimpleCollisionNode.h
--- a/Engine/Plugins/Experimental/GeometryFlow/Source/GeometryFlowMeshProcessing/Public/PhysicsNodes/GenerateSimpleCollisionNode.h
+++ b/Engine/Plugins/Experimental/GeometryFlow/Source/GeometryFlowMeshProcessing/Public/PhysicsNodes/GenerateSimpleCollisionNode.h
@@ -10,14 +10,10 @@
 #include "DataTypes/IndexSetsData.h"
 
 
-namespace UE
-{
-namespace GeometryFlow
-{
+#include "GenerateSimpleCollisionNode.generated.h"
 
-using namespace UE::Geometry;
-
-enum class ESimpleCollisionGeometryType : uint8
+UENUM()
+enum class EGeometryFlow_SimpleCollisionGeometryType : uint8
 {
 	// NOTE: This must be kept in sync with EGenerateStaticMeshLODSimpleCollisionGeometryType in GenerateStaticMeshLODProcess.h
 
@@ -31,33 +27,84 @@ enum class ESimpleCollisionGeometryType : uint8
 	None
 };
 
+UENUM()
+enum class EGeometryFlow_ProjectedHullAxisMode
+{
+	/** Use Unit X axis */
+	X = 0,
+	/** Use Unit Y axis */
+	Y = 1,
+	/** Use Unit Z axis */
+	Z = 2,
+	/** Use X/Y/Z axis with smallest axis-aligned-bounding-box dimension */
+	SmallestBoxDimension = 3,
+	/** Compute projected hull for each of X/Y/Z axes and use the one that has the smallest volume  */
+	SmallestVolume = 4
+};
+
+
+USTRUCT()
+struct GEOMETRYFLOWMESHPROCESSING_API FGenerateConvexHullSettings
+{
+	GENERATED_USTRUCT_BODY()
+	
+	UPROPERTY(EditAnywhere, Category = "Geometry Flow")
+	int32 SimplifyToTriangleCount = 50;
+
+	UPROPERTY(EditAnywhere, Category = "Geometry Flow")
+	bool bPrefilterVertices = true;
+
+	UPROPERTY(EditAnywhere, Category = "Geometry Flow")
+	int PrefilterGridResolution = 10;
+};
+
+USTRUCT()
+struct GEOMETRYFLOWMESHPROCESSING_API FGenerateSweptHullSettings
+{
+	GENERATED_USTRUCT_BODY()
+	
+	UPROPERTY(EditAnywhere, Category = "Geometry Flow")
+	bool bSimplifyPolygons = true;
+
+	UPROPERTY(EditAnywhere, Category = "Geometry Flow")
+	EGeometryFlow_ProjectedHullAxisMode SweepAxis = EGeometryFlow_ProjectedHullAxisMode::SmallestVolume;
+	
+	UPROPERTY(EditAnywhere, Category = "Geometry Flow")
+	float HullTolerance = 0.1;
+};
 
+USTRUCT()
 struct GEOMETRYFLOWMESHPROCESSING_API FGenerateSimpleCollisionSettings
 {
-	DECLARE_GEOMETRYFLOW_DATA_TYPE_IDENTIFIER(EMeshProcessingDataTypes::GenerateSimpleCollisionSettings);
+	GENERATED_USTRUCT_BODY()
+	DECLARE_GEOMETRYFLOW_DATA_TYPE_IDENTIFIER(UE::GeometryFlow::EMeshProcessingDataTypes::GenerateSimpleCollisionSettings);
 
-	ESimpleCollisionGeometryType Type = ESimpleCollisionGeometryType::ConvexHulls;
+	EGeometryFlow_SimpleCollisionGeometryType Type = EGeometryFlow_SimpleCollisionGeometryType::ConvexHulls;
 
-	struct FGenerateConvexHullSettings
-	{
-		int32 SimplifyToTriangleCount = 50;
-		bool bPrefilterVertices = true;
-		int PrefilterGridResolution = 10;
-	} ConvexHullSettings;
+	UPROPERTY(EditAnywhere, Category = "Geometry Flow")
+	FGenerateConvexHullSettings ConvexHullSettings;
 
-	struct FGenerateSweptHullSettings
-	{
-		bool bSimplifyPolygons = true;
-		FMeshSimpleShapeApproximation::EProjectedHullAxisMode SweepAxis = FMeshSimpleShapeApproximation::EProjectedHullAxisMode::SmallestVolume;
-		float HullTolerance = 0.1;
-	} SweptHullSettings;
+	UPROPERTY(EditAnywhere, Category = "Geometry Flow")
+	FGenerateSweptHullSettings SweptHullSettings;
 };
 
-GEOMETRYFLOW_DECLARE_SETTINGS_TYPES(FGenerateSimpleCollisionSettings, GenerateSimpleCollision);
+namespace UE
+{
+namespace GeometryFlow
+{
+
+using namespace UE::Geometry;
+
+typedef FGenerateSimpleCollisionSettings FGenerateSimpleCollisionSettings;
+
+
+GEOMETRYFLOW_DECLARE_USTRUCT_SETTINGS_TYPES(FGenerateSimpleCollisionSettings, GenerateSimpleCollision, 1);
 
 
 class GEOMETRYFLOWMESHPROCESSING_API FGenerateSimpleCollisionNode : public FNode
 {
+	static constexpr int Version = 1;
+	GEOMETRYFLOW_NODE_INTERNAL(FGenerateSimpleCollisionNode, Version, FNode)
 protected:
 	using SettingsDataType = TMovableData<FGenerateSimpleCollisionSettings, FGenerateSimpleCollisionSettings::DataTypeIdentifier>;
 
diff --git a/Engine/Plugins/Experimental/GeometryFlow/Source/GeometryFlowMeshProcessingEditor/GeometryFlowMeshProcessingEditor.Build.cs b/Engine/Plugins/Experimental/GeometryFlow/Source/GeometryFlowMeshProcessingEditor/GeometryFlowMeshProcessingEditor.Build.cs
--- a/Engine/Plugins/Experimental/GeometryFlow/Source/GeometryFlowMeshProcessingEditor/GeometryFlowMeshProcessingEditor.Build.cs
+++ b/Engine/Plugins/Experimental/GeometryFlow/Source/GeometryFlowMeshProcessingEditor/GeometryFlowMeshProcessingEditor.Build.cs
@@ -39,6 +39,7 @@ public class GeometryFlowMeshProcessingEditor : ModuleRules
 		PrivateDependencyModuleNames.AddRange(
 			new string[]
 			{
+				"CoreUObject",
 				"ModelingOperatorsEditorOnly"
 			}
 			);
diff --git a/Engine/Plugins/Experimental/GeometryFlow/Source/GeometryFlowMeshProcessingEditor/Private/GeometryFlowMeshProcessingEditorModule.cpp b/Engine/Plugins/Experimental/GeometryFlow/Source/GeometryFlowMeshProcessingEditor/Private/GeometryFlowMeshProcessingEditorModule.cpp
--- a/Engine/Plugins/Experimental/GeometryFlow/Source/GeometryFlowMeshProcessingEditor/Private/GeometryFlowMeshProcessingEditorModule.cpp
+++ b/Engine/Plugins/Experimental/GeometryFlow/Source/GeometryFlowMeshProcessingEditor/Private/GeometryFlowMeshProcessingEditorModule.cpp
@@ -2,12 +2,16 @@
 
 #include "GeometryFlowMeshProcessingEditorModule.h"
 #include "Modules/ModuleManager.h"
+#include "MeshProcessingEditorNodeRegistration.h"
 
 #define LOCTEXT_NAMESPACE "FGeometryFlowMeshProcessingEditorModule"
 
 void FGeometryFlowMeshProcessingEditorModule::StartupModule()
 {
 	// This code will execute after your module is loaded into memory; the exact timing is specified in the .uplugin file per-module
+
+	// Register the geometry nodes in this module with the node factory
+	UE::GeometryFlow::RegisterMeshProcessingEditorNodes();
 }
 
 void FGeometryFlowMeshProcessingEditorModule::ShutdownModule()
diff --git a/Engine/Plugins/Experimental/GeometryFlow/Source/GeometryFlowMeshProcessingEditor/Private/MeshProcessingEditorNodeRegistration.cpp b/Engine/Plugins/Experimental/GeometryFlow/Source/GeometryFlowMeshProcessingEditor/Private/MeshProcessingEditorNodeRegistration.cpp
new file mode 100644
--- /dev/null
+++ b/Engine/Plugins/Experimental/GeometryFlow/Source/GeometryFlowMeshProcessingEditor/Private/MeshProcessingEditorNodeRegistration.cpp
@@ -0,0 +1,13 @@
+// Copyright Epic Games, Inc. All Rights Reserved.
+
+#include "MeshProcessingEditorNodeRegistration.h"
+#include "GeometryFlowNodeFactory.h"
+
+#include "MeshProcessingNodes/MeshProcessingDataTypesEditor.h"
+#include "MeshProcessingNodes/MeshAutoGenerateUVsNode.h"
+
+void UE::GeometryFlow::RegisterMeshProcessingEditorNodes()
+{
+	FString CategoryName("UVs");
+	GEOMETRYFLOW_REGISTER_NODE_AND_SETTINGS_NODE(MeshAutoGenerateUVs)
+}
diff --git a/Engine/Plugins/Experimental/GeometryFlow/Source/GeometryFlowMeshProcessingEditor/Private/MeshProcessingNodes/MeshAutoGenerateUVsNode.cpp b/Engine/Plugins/Experimental/GeometryFlow/Source/GeometryFlowMeshProcessingEditor/Private/MeshProcessingNodes/MeshAutoGenerateUVsNode.cpp
--- a/Engine/Plugins/Experimental/GeometryFlow/Source/GeometryFlowMeshProcessingEditor/Private/MeshProcessingNodes/MeshAutoGenerateUVsNode.cpp
+++ b/Engine/Plugins/Experimental/GeometryFlow/Source/GeometryFlowMeshProcessingEditor/Private/MeshProcessingNodes/MeshAutoGenerateUVsNode.cpp
@@ -32,11 +32,11 @@ void FMeshAutoGenerateUVsNode::GenerateUVs(const FDynamicMesh3& MeshIn,
 	ParameterizeMeshOp.bEnablePacking = Settings.bAutoPack;
 	ParameterizeMeshOp.Width = ParameterizeMeshOp.Height = Settings.PackingTargetWidth;
 	
-	if (Settings.Method == EAutoUVMethod::UVAtlas)
+	if (FromUEnum(Settings.Method) == EAutoUVMethod::UVAtlas)
 	{
 		ParameterizeMeshOp.Method = UE::Geometry::EParamOpBackend::UVAtlas;
 	}
-	else if (Settings.Method == EAutoUVMethod::XAtlas)
+	else if (FromUEnum(Settings.Method) == EAutoUVMethod::XAtlas)
 	{
 		ParameterizeMeshOp.Method = UE::Geometry::EParamOpBackend::XAtlas;
 	}
diff --git a/Engine/Plugins/Experimental/GeometryFlow/Source/GeometryFlowMeshProcessingEditor/Public/MeshProcessingEditorNodeRegistration.h b/Engine/Plugins/Experimental/GeometryFlow/Source/GeometryFlowMeshProcessingEditor/Public/MeshProcessingEditorNodeRegistration.h
new file mode 100644
--- /dev/null
+++ b/Engine/Plugins/Experimental/GeometryFlow/Source/GeometryFlowMeshProcessingEditor/Public/MeshProcessingEditorNodeRegistration.h
@@ -0,0 +1,14 @@
+// Copyright Epic Games, Inc. All Rights Reserved.
+
+#pragma once
+
+namespace UE
+{
+	namespace GeometryFlow
+	{
+
+		// Registration is required for serialization.
+		void RegisterMeshProcessingEditorNodes();
+
+	}
+}
\ No newline at end of file
diff --git a/Engine/Plugins/Experimental/GeometryFlow/Source/GeometryFlowMeshProcessingEditor/Public/MeshProcessingNodes/MeshAutoGenerateUVsNode.h b/Engine/Plugins/Experimental/GeometryFlow/Source/GeometryFlowMeshProcessingEditor/Public/MeshProcessingNodes/MeshAutoGenerateUVsNode.h
--- a/Engine/Plugins/Experimental/GeometryFlow/Source/GeometryFlowMeshProcessingEditor/Public/MeshProcessingNodes/MeshAutoGenerateUVsNode.h
+++ b/Engine/Plugins/Experimental/GeometryFlow/Source/GeometryFlowMeshProcessingEditor/Public/MeshProcessingNodes/MeshAutoGenerateUVsNode.h
@@ -2,16 +2,16 @@
 
 #pragma once
 
+#include "GeometryFlowTSourceNode.h"
 #include "MeshProcessingNodes/MeshProcessingBaseNodes.h"
 #include "MeshProcessingNodes/MeshProcessingDataTypesEditor.h"
+#include "Serialization/Archive.h"
 
-namespace UE
-{
-namespace GeometryFlow
-{
 
+#include "MeshAutoGenerateUVsNode.generated.h"
 
-enum class EAutoUVMethod : uint8
+UENUM()
+enum class EGeometryFlow_AutoUVMethod : uint8
 {
 	PatchBuilder = 0,
 	UVAtlas = 1,
@@ -19,37 +19,86 @@ enum class EAutoUVMethod : uint8
 };
 
 
-
+USTRUCT()
 struct GEOMETRYFLOWMESHPROCESSINGEDITOR_API FMeshAutoGenerateUVsSettings
 {
-	DECLARE_GEOMETRYFLOW_DATA_TYPE_IDENTIFIER(EMeshProcessingDataTypesEditor::MeshAutoGenerateUVsSettings);
+	GENERATED_USTRUCT_BODY()
 
-	EAutoUVMethod Method = EAutoUVMethod::PatchBuilder;
+	DECLARE_GEOMETRYFLOW_DATA_TYPE_IDENTIFIER(UE::GeometryFlow::EMeshProcessingDataTypesEditor::MeshAutoGenerateUVsSettings);
+
+	UPROPERTY(EditAnywhere, Category = "Geometry Flow")
+	EGeometryFlow_AutoUVMethod Method = EGeometryFlow_AutoUVMethod::PatchBuilder;
 
 	// UVAtlas parameters
+	UPROPERTY(EditAnywhere, Category = "Geometry Flow")
 	double UVAtlasStretch = 0.5;
+
+	UPROPERTY(EditAnywhere, Category = "Geometry Flow")
 	int UVAtlasNumCharts = 0;
 
 	// XAtlas parameters
+	UPROPERTY(EditAnywhere, Category = "Geometry Flow")
 	int XAtlasMaxIterations = 1;
 
 	// PatchBuilder parameters
+	UPROPERTY(EditAnywhere, Category = "Geometry Flow")
 	int NumInitialPatches = 100;
+
+	UPROPERTY(EditAnywhere, Category = "Geometry Flow")
 	double CurvatureAlignment = 1.0;
+
+	UPROPERTY(EditAnywhere, Category = "Geometry Flow")
 	double MergingThreshold = 1.5;
+
+	UPROPERTY(EditAnywhere, Category = "Geometry Flow")
 	double MaxAngleDeviationDeg = 45.0;
+
+	UPROPERTY(EditAnywhere, Category = "Geometry Flow")
 	int SmoothingSteps = 5;
+
+	UPROPERTY(EditAnywhere, Category = "Geometry Flow")
 	double SmoothingAlpha = 0.25;
 
+	UPROPERTY(EditAnywhere, Category = "Geometry Flow")
 	bool bAutoPack = false;
+
+	UPROPERTY(EditAnywhere, Category = "Geometry Flow")
 	int PackingTargetWidth = 512;
 };
-GEOMETRYFLOW_DECLARE_SETTINGS_TYPES(FMeshAutoGenerateUVsSettings, MeshAutoGenerateUVs);
+
+
+namespace UE
+{
+namespace GeometryFlow
+{
+// bring into the namespace
+
+typedef FMeshAutoGenerateUVsSettings FMeshAutoGenerateUVsSettings;
+
+enum class EAutoUVMethod : uint8
+{
+	PatchBuilder = 0,
+	UVAtlas = 1,
+	XAtlas = 2
+};
+static EAutoUVMethod FromUEnum(const  EGeometryFlow_AutoUVMethod& UVsUnwrapType)
+{
+	return static_cast<EAutoUVMethod>(UVsUnwrapType);
+}
+static EGeometryFlow_AutoUVMethod ToUEnum(const EAutoUVMethod& UVsUnwrapType)
+{
+	return static_cast<EGeometryFlow_AutoUVMethod>(UVsUnwrapType);
+}
+
+
+GEOMETRYFLOW_DECLARE_USTRUCT_SETTINGS_TYPES(FMeshAutoGenerateUVsSettings, MeshAutoGenerateUVs, 1);
 
 
 
 class GEOMETRYFLOWMESHPROCESSINGEDITOR_API FMeshAutoGenerateUVsNode : public TProcessMeshWithSettingsBaseNode<FMeshAutoGenerateUVsSettings>
 {
+	static constexpr int Version = 1;
+	GEOMETRYFLOW_NODE_INTERNAL(FMeshAutoGenerateUVsNode, Version, FNode)
 public:
 	FMeshAutoGenerateUVsNode() : TProcessMeshWithSettingsBaseNode<FMeshAutoGenerateUVsSettings>()
 	{
