Fix uninitialized member in FMeshMakeBakingCacheSettings
CIS issue 600482

diff --git a/Engine/Plugins/Experimental/GeometryFlow/Source/GeometryFlowMeshProcessing/Public/DataTypes/MeshImageBakingData.h b/Engine/Plugins/Experimental/GeometryFlow/Source/GeometryFlowMeshProcessing/Public/DataTypes/MeshImageBakingData.h
--- a/Engine/Plugins/Experimental/GeometryFlow/Source/GeometryFlowMeshProcessing/Public/DataTypes/MeshImageBakingData.h
+++ b/Engine/Plugins/Experimental/GeometryFlow/Source/GeometryFlowMeshProcessing/Public/DataTypes/MeshImageBakingData.h
@@ -22,7 +22,7 @@ struct GEOMETRYFLOWMESHPROCESSING_API FMeshMakeBakingCacheSettings
 
 
 	UPROPERTY(EditAnywhere, Category = "Geometry Flow")
-	FIntPoint Dimensions;  //FImageDimensions Dimensions;
+	FIntPoint Dimensions{ EForceInit::ForceInit };  //FImageDimensions Dimensions;
 	
 	UPROPERTY(EditAnywhere, Category = "Geometry Flow")
 	int32 UVLayer = 0;
