AutoLOD tool - fix bugs I introduced.

diff --git a/Engine/Plugins/Experimental/GeometryFlow/Source/GeometryFlowCore/Public/BaseNodes/SwitchNode.h b/Engine/Plugins/Experimental/GeometryFlow/Source/GeometryFlowCore/Public/BaseNodes/SwitchNode.h
--- a/Engine/Plugins/Experimental/GeometryFlow/Source/GeometryFlowCore/Public/BaseNodes/SwitchNode.h
+++ b/Engine/Plugins/Experimental/GeometryFlow/Source/GeometryFlowCore/Public/BaseNodes/SwitchNode.h
@@ -26,6 +26,19 @@ namespace GeometryFlow
 template<typename T, int NumInputs, int StorageTypeIdentifier>
 class TSwitchNode : public FNode
 {
+
+	static constexpr int NodeVersion = 1;
+public:
+	typedef FNode Super;
+	static FName StaticType() 
+	{ 
+		FString NodeTypeName = FString::Printf(TEXT("FSwitchNode%02dInputs%02dType"), NumInputs, StorageTypeIdentifier);
+		return FName(*NodeTypeName); 
+	}
+	virtual FName GetType() const override { return StaticType(); }
+	virtual bool  IsA(FName InType) const override { return InType.ToString().Equals(StaticType().ToString()) || Super::IsA(InType); }
+	virtual int32 GetVersionID() const override { return NodeVersion; }
+
 protected:
 	using DataContainerType = TMovableData<T, StorageTypeIdentifier>;
 
@@ -156,9 +169,6 @@ public:
 	}
 };
 
-
-
-
 }	// end namespace GeometryFlow
 }	// end namespace UE
 
diff --git a/Engine/Plugins/Experimental/GeometryFlow/Source/GeometryFlowMeshProcessing/Private/MeshProcessingNodeRegistration.cpp b/Engine/Plugins/Experimental/GeometryFlow/Source/GeometryFlowMeshProcessing/Private/MeshProcessingNodeRegistration.cpp
--- a/Engine/Plugins/Experimental/GeometryFlow/Source/GeometryFlowMeshProcessing/Private/MeshProcessingNodeRegistration.cpp
+++ b/Engine/Plugins/Experimental/GeometryFlow/Source/GeometryFlowMeshProcessing/Private/MeshProcessingNodeRegistration.cpp
@@ -113,6 +113,7 @@ void RegisterMeshProcessingNodes()
 
 	GEOMETRYFLOW_REGISTER_SETTINGS_NODE_TYPE(Tangents)
 	GEOMETRYFLOW_REGISTER_NODE_TYPE(ComputeMeshTangents)
+	GEOMETRYFLOW_REGISTER_NODE_TYPE(MeshTangentsTransfer)
 
 	GEOMETRYFLOW_REGISTER_SETTINGS_NODE_TYPE(Thicken)
 	GEOMETRYFLOW_REGISTER_NODE_TYPE(MeshThicken)
@@ -131,7 +132,7 @@ void RegisterMeshProcessingNodes()
 
 
 	// PhysicsNodes /
-	CategoryName = FString("Phyics");
+	CategoryName = FString("Physics");
 	GEOMETRYFLOW_REGISTER_NODE_AND_SETTINGS_NODE(GenerateSimpleCollision)
 
 
diff --git a/Engine/Plugins/Experimental/GeometryFlow/Source/GeometryFlowMeshProcessing/Public/DataTypes/DynamicMeshData.h b/Engine/Plugins/Experimental/GeometryFlow/Source/GeometryFlowMeshProcessing/Public/DataTypes/DynamicMeshData.h
--- a/Engine/Plugins/Experimental/GeometryFlow/Source/GeometryFlowMeshProcessing/Public/DataTypes/DynamicMeshData.h
+++ b/Engine/Plugins/Experimental/GeometryFlow/Source/GeometryFlowMeshProcessing/Public/DataTypes/DynamicMeshData.h
@@ -4,6 +4,7 @@
 
 #include "DynamicMesh/DynamicMesh3.h"
 #include "GeometryFlowCoreNodes.h"
+#include "BaseNodes/TransferNode.h"
 #include "MeshProcessingNodes/MeshProcessingDataTypes.h"
 
 namespace UE::Geometry { class FDynamicMesh3; }
@@ -29,8 +30,11 @@ class FDynamicMeshSourceNode : public  TSourceNodeBase<FDynamicMesh3, (int)EMesh
 };
 
 
-typedef TTransferNode<FDynamicMesh3, (int)EMeshProcessingDataTypes::DynamicMesh> FDynamicMeshTransferNode;
-
+class FDynamicMeshTransferNode : public TTransferNode<FDynamicMesh3, (int)EMeshProcessingDataTypes::DynamicMesh>
+{
+	static constexpr int Version = 1;
+	GEOMETRYFLOW_NODE_INTERNAL(FDynamicMeshTransferNode, Version, FNode)
+};
 
 }	// end namespace GeometryFlow
 }	// end namespace UE
diff --git a/Engine/Plugins/Experimental/GeometryFlow/Source/GeometryFlowMeshProcessing/Public/DataTypes/MeshImageBakingData.h b/Engine/Plugins/Experimental/GeometryFlow/Source/GeometryFlowMeshProcessing/Public/DataTypes/MeshImageBakingData.h
--- a/Engine/Plugins/Experimental/GeometryFlow/Source/GeometryFlowMeshProcessing/Public/DataTypes/MeshImageBakingData.h
+++ b/Engine/Plugins/Experimental/GeometryFlow/Source/GeometryFlowMeshProcessing/Public/DataTypes/MeshImageBakingData.h
@@ -68,7 +68,8 @@ class GEOMETRYFLOWMESHPROCESSING_API FMakeMeshBakingCacheNode : public FNode
 {
 protected:
 	using SettingsDataType = TMovableData<FMeshMakeBakingCacheSettings, FMeshMakeBakingCacheSettings::DataTypeIdentifier>;
-
+	static constexpr int Version = 1;
+	GEOMETRYFLOW_NODE_INTERNAL(FMakeMeshBakingCacheNode, Version, FNode)
 public:
 	static const FString InParamDetailMesh() { return TEXT("DetailMesh"); }
 	static const FString InParamTargetMesh() { return TEXT("TargetMesh"); }
