Geometry Processing Unit Tests:  eliminate warnings and other problems in preparation for enabling the tests on mac or linux

diff --git a/Engine/Plugins/Experimental/GeometryFlow/Source/GeometryFlowMeshProcessing/Public/DataTypes/MeshImageBakingData.h b/Engine/Plugins/Experimental/GeometryFlow/Source/GeometryFlowMeshProcessing/Public/DataTypes/MeshImageBakingData.h
--- a/Engine/Plugins/Experimental/GeometryFlow/Source/GeometryFlowMeshProcessing/Public/DataTypes/MeshImageBakingData.h
+++ b/Engine/Plugins/Experimental/GeometryFlow/Source/GeometryFlowMeshProcessing/Public/DataTypes/MeshImageBakingData.h
@@ -28,7 +28,7 @@ struct GEOMETRYFLOWMESHPROCESSING_API FMeshMakeBakingCacheSettings
 	int32 UVLayer = 0;
 
 	UPROPERTY(EditAnywhere, Category = "Geometry Flow")
-	float Thickness = 0.1;
+	float Thickness = 0.1f;
 
 };
 
diff --git a/Engine/Plugins/Experimental/GeometryFlow/Source/GeometryFlowMeshProcessing/Public/PhysicsNodes/GenerateSimpleCollisionNode.h b/Engine/Plugins/Experimental/GeometryFlow/Source/GeometryFlowMeshProcessing/Public/PhysicsNodes/GenerateSimpleCollisionNode.h
--- a/Engine/Plugins/Experimental/GeometryFlow/Source/GeometryFlowMeshProcessing/Public/PhysicsNodes/GenerateSimpleCollisionNode.h
+++ b/Engine/Plugins/Experimental/GeometryFlow/Source/GeometryFlowMeshProcessing/Public/PhysicsNodes/GenerateSimpleCollisionNode.h
@@ -70,7 +70,7 @@ struct GEOMETRYFLOWMESHPROCESSING_API FGenerateSweptHullSettings
 	EGeometryFlow_ProjectedHullAxisMode SweepAxis = EGeometryFlow_ProjectedHullAxisMode::SmallestVolume;
 	
 	UPROPERTY(EditAnywhere, Category = "Geometry Flow")
-	float HullTolerance = 0.1;
+	float HullTolerance = 0.1f;
 };
 
 USTRUCT()
diff --git a/Engine/Source/Runtime/GeometryCore/Public/Implicit/CachingMeshSDF.h b/Engine/Source/Runtime/GeometryCore/Public/Implicit/CachingMeshSDF.h
--- a/Engine/Source/Runtime/GeometryCore/Public/Implicit/CachingMeshSDF.h
+++ b/Engine/Source/Runtime/GeometryCore/Public/Implicit/CachingMeshSDF.h
@@ -156,7 +156,7 @@ public:
 		int NJ = (int)((max.Y - GridOrigin.Y) / CellSize) + 1;
 		int NK = (int)((max.Z - GridOrigin.Z) / CellSize) + 1;
 
-		UpperBoundDistance = (float)((NI + NJ + NK) * CellSize);
+		UpperBoundDistance = ((float)(NI + NJ + NK) * CellSize);
 		Grid = FDenseGrid3f(NI, NJ, NK, UpperBoundDistance);
 
 		MaxDistQueryDist = MaxOffsetDistance + (2 * CellSize * FMathf::Sqrt2);
diff --git a/Engine/Source/Runtime/GeometryCore/Public/Spatial/MeshWindingNumberGrid.h b/Engine/Source/Runtime/GeometryCore/Public/Spatial/MeshWindingNumberGrid.h
--- a/Engine/Source/Runtime/GeometryCore/Public/Spatial/MeshWindingNumberGrid.h
+++ b/Engine/Source/Runtime/GeometryCore/Public/Spatial/MeshWindingNumberGrid.h
@@ -87,7 +87,7 @@ public:
 			return;
 		}
 
-		float fBufferWidth = BufferCells * (float)CellSize;
+		float fBufferWidth = (float)BufferCells * (float)CellSize;
 		GridOrigin = (FVector3f)bounds.Min - fBufferWidth * FVector3f::One();
 		FVector3f max = (FVector3f)bounds.Max + fBufferWidth * FVector3f::One();
 		int ni = (int)((max.X - GridOrigin.X) / (float)CellSize) + 1;
