GeometryFlow:  clean-up node registration

diff --git a/Engine/Plugins/Editor/MeshLODToolset/Source/MeshLODToolset/Private/Graphs/GenerateMeshLODGraph.cpp b/Engine/Plugins/Editor/MeshLODToolset/Source/MeshLODToolset/Private/Graphs/GenerateMeshLODGraph.cpp
--- a/Engine/Plugins/Editor/MeshLODToolset/Source/MeshLODToolset/Private/Graphs/GenerateMeshLODGraph.cpp
+++ b/Engine/Plugins/Editor/MeshLODToolset/Source/MeshLODToolset/Private/Graphs/GenerateMeshLODGraph.cpp
@@ -19,6 +19,9 @@
 #include "MeshBakingNodes/BakeMeshMultiTextureNode.h"
 
 #include "MeshDecompositionNodes/MakeTriangleSetsNode.h"
+#include "MeshProcessingNodeRegistration.h"
+#include "MeshProcessingEditorNodeRegistration.h"
+#include "GeometryFlowCoreNodeRegistration.h"
 
 using namespace UE::Geometry;
 using namespace UE::GeometryFlow;
@@ -293,6 +296,13 @@ static bool IsSingleMaterialMesh(const FDynamicMesh3& Mesh, int32& UniqueMateria
 
 void FGenerateMeshLODGraph::BuildGraph(const FDynamicMesh3* SourceMeshHint)
 {
+	// does nothing if all the node types are already registered.
+	{
+		UE::GeometryFlow::FCoreNodeRegistration::RegisterNodes();
+		UE::GeometryFlow::FMeshProcessingNodeRegistration::RegisterNodes();
+		UE::GeometryFlow::FMeshProcessingEditorNodeRegistration::RegisterNodes();
+	}
+
 	// precompute hints (async?)
 	int32 UniqueMaterialID = 0;
 	bool bIsSingleMaterialMesh = (SourceMeshHint != nullptr) ? IsSingleMaterialMesh(*SourceMeshHint, UniqueMaterialID) : false;
diff --git a/Engine/Plugins/Experimental/GeometryFlow/Source/GeometryFlowCore/Private/GeometryFlowCoreModule.cpp b/Engine/Plugins/Experimental/GeometryFlow/Source/GeometryFlowCore/Private/GeometryFlowCoreModule.cpp
--- a/Engine/Plugins/Experimental/GeometryFlow/Source/GeometryFlowCore/Private/GeometryFlowCoreModule.cpp
+++ b/Engine/Plugins/Experimental/GeometryFlow/Source/GeometryFlowCore/Private/GeometryFlowCoreModule.cpp
@@ -3,7 +3,7 @@
 #include "GeometryFlowCoreModule.h"
 #include "Modules/ModuleManager.h"
 
-#include "GeometryFlowNodeFactory.h"
+#include "GeometryFlowCoreNodeRegistration.h"
 
 #define LOCTEXT_NAMESPACE "FGeometryFlowCoreModule"
 
@@ -12,7 +12,7 @@ void FGeometryFlowCoreModule::StartupModule()
 	// This code will execute after your module is loaded into memory; the exact timing is specified in the .uplugin file per-module
 
 		// Register the nodes defined by this module with the GeometryFlowNodeFactory.
-	UE::GeometryFlow::RegisterCoreNodes();
+	UE::GeometryFlow::FCoreNodeRegistration::RegisterNodes();
 }
 
 void FGeometryFlowCoreModule::ShutdownModule()
diff --git a/Engine/Plugins/Experimental/GeometryFlow/Source/GeometryFlowCore/Private/GeometryFlowGraph.cpp b/Engine/Plugins/Experimental/GeometryFlow/Source/GeometryFlowCore/Private/GeometryFlowGraph.cpp
--- a/Engine/Plugins/Experimental/GeometryFlow/Source/GeometryFlowCore/Private/GeometryFlowGraph.cpp
+++ b/Engine/Plugins/Experimental/GeometryFlow/Source/GeometryFlowCore/Private/GeometryFlowGraph.cpp
@@ -7,6 +7,7 @@
 
 using namespace UE::GeometryFlow;
 
+DEFINE_LOG_CATEGORY(LogGeometryFlowGraph);
 
 namespace
 {
@@ -75,6 +76,26 @@ ENodeCachingStrategy FGraph::GetCachingStrategyForNode(FHandle NodeHandle) const
 		DefaultCachingStrategy : Found->CachingStrategy;
 }
 
+FGraph::FHandle FGraph::AddNodeOfType(
+	const FName TypeName,
+	const FString& Identifier,
+	ENodeCachingStrategy CachingStrategy)
+{
+	FHandle Handle;
+	Handle.Identifier = NodeCounter + 1;
+
+	if (AddNodeOfType(TypeName, Handle, Identifier, CachingStrategy) == ENodeAddResult::Success)
+	{
+		NodeCounter++;
+	}
+	else
+	{
+		Handle = FHandle(); // invalid handle.
+		UE_LOG(LogGeometryFlowGraph, Warning, TEXT("Unregisterd type '%s' : Failed to add node '%s' "), *TypeName.ToString() , *Identifier);
+	}
+
+	return Handle;
+}
 
 bool  FGraph::IsDependent(FHandle DependentNode, FHandle IndependentNode) const
 {
diff --git a/Engine/Plugins/Experimental/GeometryFlow/Source/GeometryFlowCore/Private/GeometryFlowNodeFactory.cpp b/Engine/Plugins/Experimental/GeometryFlow/Source/GeometryFlowCore/Private/GeometryFlowNodeFactory.cpp
new file mode 100644
--- /dev/null
+++ b/Engine/Plugins/Experimental/GeometryFlow/Source/GeometryFlowCore/Private/GeometryFlowNodeFactory.cpp
@@ -0,0 +1,15 @@
+// Copyright Epic Games, Inc. All Rights Reserved.
+
+#include "GeometryFlowNodeFactory.h"
+#include "Misc/LazySingleton.h"
+
+
+UE::GeometryFlow::FNodeFactory& UE::GeometryFlow::FNodeFactory::GetInstance()
+{
+	return TLazySingleton< UE::GeometryFlow::FNodeFactory >::Get();
+};
+
+void UE::GeometryFlow::FNodeFactory::TearDown()
+{
+	TLazySingleton< UE::GeometryFlow::FNodeFactory >::TearDown();
+};
\ No newline at end of file
diff --git a/Engine/Plugins/Experimental/GeometryFlow/Source/GeometryFlowCore/Private/GometryFlowCoreNodeRegistration.cpp b/Engine/Plugins/Experimental/GeometryFlow/Source/GeometryFlowCore/Private/GometryFlowCoreNodeRegistration.cpp
--- a/Engine/Plugins/Experimental/GeometryFlow/Source/GeometryFlowCore/Private/GometryFlowCoreNodeRegistration.cpp
+++ b/Engine/Plugins/Experimental/GeometryFlow/Source/GeometryFlowCore/Private/GometryFlowCoreNodeRegistration.cpp
@@ -1,5 +1,6 @@
 // Copyright Epic Games, Inc. All Rights Reserved.
 
+#include "GeometryFlowCoreNodeRegistration.h"
 #include "GeometryFlowNodeFactory.h"
 #include "GeometryFlowCoreNodes.h"
 
@@ -23,7 +24,7 @@ namespace GeometryFlow
 																	  CategoryName);
 
 	// Register Core Nodes with the node factory
-	void RegisterCoreNodes()
+	void FCoreNodeRegistration::RegisterNodes()
 	{	
 		FString CategoryName("Basic Types");
 
diff --git a/Engine/Plugins/Experimental/GeometryFlow/Source/GeometryFlowCore/Public/GeometryFlowCoreNodeRegistration.h b/Engine/Plugins/Experimental/GeometryFlow/Source/GeometryFlowCore/Public/GeometryFlowCoreNodeRegistration.h
new file mode 100644
--- /dev/null
+++ b/Engine/Plugins/Experimental/GeometryFlow/Source/GeometryFlowCore/Public/GeometryFlowCoreNodeRegistration.h
@@ -0,0 +1,18 @@
+// Copyright Epic Games, Inc. All Rights Reserved.
+
+#pragma once
+
+#include "CoreMinimal.h"
+
+namespace UE
+{
+namespace GeometryFlow
+{
+
+	struct FCoreNodeRegistration
+	{
+		// function that registers the core nodes defined in GeometryFlowCore (mostly POD source nodes)
+		GEOMETRYFLOWCORE_API static void RegisterNodes();
+	};
+}
+}
\ No newline at end of file
diff --git a/Engine/Plugins/Experimental/GeometryFlow/Source/GeometryFlowCore/Public/GeometryFlowGraph.h b/Engine/Plugins/Experimental/GeometryFlow/Source/GeometryFlowCore/Public/GeometryFlowGraph.h
--- a/Engine/Plugins/Experimental/GeometryFlow/Source/GeometryFlowCore/Public/GeometryFlowGraph.h
+++ b/Engine/Plugins/Experimental/GeometryFlow/Source/GeometryFlowCore/Public/GeometryFlowGraph.h
@@ -9,6 +9,8 @@
 
 class FProgressCancel;
 
+GEOMETRYFLOWCORE_API DECLARE_LOG_CATEGORY_EXTERN(LogGeometryFlowGraph, Display, All);
+
 namespace UE
 {
 namespace GeometryFlow
@@ -132,22 +134,7 @@ public:
 	FHandle AddNodeOfType(
 		const FName TypeName,
 		const FString& Identifier = FString(""),
-		ENodeCachingStrategy CachingStrategy = ENodeCachingStrategy::Default)
-	{
-		FHandle Handle;
-		Handle.Identifier= NodeCounter + 1;
-
-		if (AddNodeOfType(TypeName, Handle, Identifier, CachingStrategy) == ENodeAddResult::Success)
-		{
-			NodeCounter++;
-		}
-		else
-		{
-			Handle = FHandle(); // invalid handle.
-		}
-
-		return Handle;
-	}
+		ENodeCachingStrategy CachingStrategy = ENodeCachingStrategy::Default);
 
 	// Note: once removed, any external handles to this node will be stale. 
 	void RemoveNode(const FHandle Handle)
diff --git a/Engine/Plugins/Experimental/GeometryFlow/Source/GeometryFlowCore/Public/GeometryFlowNodeFactory.h b/Engine/Plugins/Experimental/GeometryFlow/Source/GeometryFlowCore/Public/GeometryFlowNodeFactory.h
--- a/Engine/Plugins/Experimental/GeometryFlow/Source/GeometryFlowCore/Public/GeometryFlowNodeFactory.h
+++ b/Engine/Plugins/Experimental/GeometryFlow/Source/GeometryFlowCore/Public/GeometryFlowNodeFactory.h
@@ -5,6 +5,9 @@
 #include "GeometryFlowNode.h"
 #include "CoreMinimal.h"
 
+
+class FLazySingleton;
+
 namespace UE
 {
 namespace GeometryFlow
@@ -12,7 +15,7 @@ namespace GeometryFlow
 
 	// Singleton used to create nodes UE::GeometryFlow::FGraph, introduced to help with serialization.  
 	// this contains an map between FNames and function pointers that can create the named nodes.
-	class GEOMETRYFLOWCORE_API FNodeFactory
+	class  FNodeFactory
 	{
 
 	public:
@@ -29,12 +32,10 @@ namespace GeometryFlow
 
 		FNodeFactory(const FNodeFactory&) = delete;
 	
-		static FNodeFactory& GetInstance()
-		{
-			// Meyers style singleton
-			static FNodeFactory Instance;
-			return Instance;
-		};
+	
+
+		GEOMETRYFLOWCORE_API static  FNodeFactory& GetInstance();
+		GEOMETRYFLOWCORE_API static  void TearDown();
 
 		// returns false if the NodeTypeName is already registered.
 		template <typename NodeType>
@@ -68,7 +69,7 @@ namespace GeometryFlow
 		}
 
 
-		bool CanMake(const FName NodeTypeName) const
+		GEOMETRYFLOWCORE_API bool CanMake(const FName NodeTypeName) const
 		{
 			return Factories.Contains(NodeTypeName);
 		}
@@ -85,19 +86,22 @@ namespace GeometryFlow
 			return Result;
 		}
 
-		const TMap<const FName, FNodeFactory::FNodeFactoryInfo>& GetFactories() const { return Factories;}
+		GEOMETRYFLOWCORE_API const TMap<const FName, FNodeFactory::FNodeFactoryInfo>& GetFactories() const { return Factories;}
 
 	protected:
 
 		TMap<const FName, FNodeFactoryInfo>  Factories;
 		
+		
+		friend FLazySingleton; // needed because the FNodeFactory constructor is protected.
 		FNodeFactory()
 		{}
+
+	
 	};
 
 	
-	// function that registers the core nodes defined in GeometryFlowCore (mostly POD source nodes)
-	void RegisterCoreNodes();
+
 }
 
 
diff --git a/Engine/Plugins/Experimental/GeometryFlow/Source/GeometryFlowMeshProcessing/Private/GeometryFlowMeshProcessingModule.cpp b/Engine/Plugins/Experimental/GeometryFlow/Source/GeometryFlowMeshProcessing/Private/GeometryFlowMeshProcessingModule.cpp
--- a/Engine/Plugins/Experimental/GeometryFlow/Source/GeometryFlowMeshProcessing/Private/GeometryFlowMeshProcessingModule.cpp
+++ b/Engine/Plugins/Experimental/GeometryFlow/Source/GeometryFlowMeshProcessing/Private/GeometryFlowMeshProcessingModule.cpp
@@ -12,8 +12,8 @@ void FGeometryFlowMeshProcessingModule::StartupModule()
 {
 	// This code will execute after your module is loaded into memory; the exact timing is specified in the .uplugin file per-module
 	
-	UE::GeometryFlow::RegisterMeshProcessingNodes();
-}
+	UE::GeometryFlow::FMeshProcessingNodeRegistration::RegisterNodes();
+} 
 
 void FGeometryFlowMeshProcessingModule::ShutdownModule()
 {
diff --git a/Engine/Plugins/Experimental/GeometryFlow/Source/GeometryFlowMeshProcessing/Private/MeshProcessingNodeRegistration.cpp b/Engine/Plugins/Experimental/GeometryFlow/Source/GeometryFlowMeshProcessing/Private/MeshProcessingNodeRegistration.cpp
--- a/Engine/Plugins/Experimental/GeometryFlow/Source/GeometryFlowMeshProcessing/Private/MeshProcessingNodeRegistration.cpp
+++ b/Engine/Plugins/Experimental/GeometryFlow/Source/GeometryFlowMeshProcessing/Private/MeshProcessingNodeRegistration.cpp
@@ -50,7 +50,7 @@ namespace GeometryFlow
 {
 
 																			  
-void RegisterMeshProcessingNodes()
+void FMeshProcessingNodeRegistration::RegisterNodes()
 {
 
 	FString CategoryName("Data Types");
diff --git a/Engine/Plugins/Experimental/GeometryFlow/Source/GeometryFlowMeshProcessing/Public/MeshProcessingNodeRegistration.h b/Engine/Plugins/Experimental/GeometryFlow/Source/GeometryFlowMeshProcessing/Public/MeshProcessingNodeRegistration.h
--- a/Engine/Plugins/Experimental/GeometryFlow/Source/GeometryFlowMeshProcessing/Public/MeshProcessingNodeRegistration.h
+++ b/Engine/Plugins/Experimental/GeometryFlow/Source/GeometryFlowMeshProcessing/Public/MeshProcessingNodeRegistration.h
@@ -2,13 +2,19 @@
 
 #pragma once
 
+#include "CoreMinimal.h"
+
 namespace UE
 {
 namespace GeometryFlow
 {
 
-	// Required registration of the geometry flow graph nodes in the MeshProcessing Module.
-	void RegisterMeshProcessingNodes();
+	struct FMeshProcessingNodeRegistration
+	{
+		// Required registration of the geometry flow graph nodes in the MeshProcessing Module.
+		GEOMETRYFLOWMESHPROCESSING_API static void RegisterNodes();
+	};
+
 
 }
 }
\ No newline at end of file
diff --git a/Engine/Plugins/Experimental/GeometryFlow/Source/GeometryFlowMeshProcessingEditor/Private/GeometryFlowMeshProcessingEditorModule.cpp b/Engine/Plugins/Experimental/GeometryFlow/Source/GeometryFlowMeshProcessingEditor/Private/GeometryFlowMeshProcessingEditorModule.cpp
--- a/Engine/Plugins/Experimental/GeometryFlow/Source/GeometryFlowMeshProcessingEditor/Private/GeometryFlowMeshProcessingEditorModule.cpp
+++ b/Engine/Plugins/Experimental/GeometryFlow/Source/GeometryFlowMeshProcessingEditor/Private/GeometryFlowMeshProcessingEditorModule.cpp
@@ -11,7 +11,7 @@ void FGeometryFlowMeshProcessingEditorModule::StartupModule()
 	// This code will execute after your module is loaded into memory; the exact timing is specified in the .uplugin file per-module
 
 	// Register the geometry nodes in this module with the node factory
-	UE::GeometryFlow::RegisterMeshProcessingEditorNodes();
+	UE::GeometryFlow::FMeshProcessingEditorNodeRegistration::RegisterNodes();
 }
 
 void FGeometryFlowMeshProcessingEditorModule::ShutdownModule()
diff --git a/Engine/Plugins/Experimental/GeometryFlow/Source/GeometryFlowMeshProcessingEditor/Private/MeshProcessingEditorNodeRegistration.cpp b/Engine/Plugins/Experimental/GeometryFlow/Source/GeometryFlowMeshProcessingEditor/Private/MeshProcessingEditorNodeRegistration.cpp
--- a/Engine/Plugins/Experimental/GeometryFlow/Source/GeometryFlowMeshProcessingEditor/Private/MeshProcessingEditorNodeRegistration.cpp
+++ b/Engine/Plugins/Experimental/GeometryFlow/Source/GeometryFlowMeshProcessingEditor/Private/MeshProcessingEditorNodeRegistration.cpp
@@ -6,7 +6,7 @@
 #include "MeshProcessingNodes/MeshProcessingDataTypesEditor.h"
 #include "MeshProcessingNodes/MeshAutoGenerateUVsNode.h"
 
-void UE::GeometryFlow::RegisterMeshProcessingEditorNodes()
+void UE::GeometryFlow::FMeshProcessingEditorNodeRegistration::RegisterNodes()
 {
 	FString CategoryName("UVs");
 	GEOMETRYFLOW_REGISTER_NODE_AND_SETTINGS_NODE(MeshAutoGenerateUVs)
diff --git a/Engine/Plugins/Experimental/GeometryFlow/Source/GeometryFlowMeshProcessingEditor/Public/MeshProcessingEditorNodeRegistration.h b/Engine/Plugins/Experimental/GeometryFlow/Source/GeometryFlowMeshProcessingEditor/Public/MeshProcessingEditorNodeRegistration.h
--- a/Engine/Plugins/Experimental/GeometryFlow/Source/GeometryFlowMeshProcessingEditor/Public/MeshProcessingEditorNodeRegistration.h
+++ b/Engine/Plugins/Experimental/GeometryFlow/Source/GeometryFlowMeshProcessingEditor/Public/MeshProcessingEditorNodeRegistration.h
@@ -2,13 +2,17 @@
 
 #pragma once
 
+#include "CoreMinimal.h"
+
 namespace UE
 {
 	namespace GeometryFlow
 	{
 
-		// Registration is required for serialization.
-		void RegisterMeshProcessingEditorNodes();
-
+		struct FMeshProcessingEditorNodeRegistration
+		{
+			// Registration is required for serialization.
+			GEOMETRYFLOWMESHPROCESSINGEDITOR_API static void  RegisterNodes();
+		};
 	}
 }
\ No newline at end of file
